Added whitespace validation to LocRes compilation, and added flags to control which validation to run on a per-target basis

diff --git a/Engine/Source/Developer/Localization/Private/LocalizationConfigurationScript.cpp b/Engine/Source/Developer/Localization/Private/LocalizationConfigurationScript.cpp
--- a/Engine/Source/Developer/Localization/Private/LocalizationConfigurationScript.cpp
+++ b/Engine/Source/Developer/Localization/Private/LocalizationConfigurationScript.cpp
@@ -1068,6 +1068,8 @@ namespace LocalizationConfigurationScript
 			ConfigSection.Add( TEXT("ResourceName"), GetLocResFileName(Target) );
 
 			ConfigSection.Add( TEXT("bSkipSourceCheck"), Target->Settings.CompileSettings.SkipSourceCheck ? TEXT("true") : TEXT("false") );
+			ConfigSection.Add( TEXT("bValidateFormatPatterns"), Target->Settings.CompileSettings.ValidateFormatPatterns ? TEXT("true") : TEXT("false") );
+			ConfigSection.Add( TEXT("bValidateSafeWhitespace"), Target->Settings.CompileSettings.ValidateSafeWhitespace ? TEXT("true") : TEXT("false") );
 
 			if (Target->Settings.SupportedCulturesStatistics.IsValidIndex(Target->Settings.NativeCultureIndex))
 			{
diff --git a/Engine/Source/Developer/Localization/Private/LocalizationModule.cpp b/Engine/Source/Developer/Localization/Private/LocalizationModule.cpp
--- a/Engine/Source/Developer/Localization/Private/LocalizationModule.cpp
+++ b/Engine/Source/Developer/Localization/Private/LocalizationModule.cpp
@@ -7,9 +7,9 @@
 class FLocalizationModule : public ILocalizationModule
 {
 public:
-	virtual bool HandleRegenLocCommand(const FString& InConfigFilePath, const bool bSkipSourceCheck) override
+	virtual bool HandleRegenLocCommand(const FString& InConfigFilePath) override
 	{
-		return FTextLocalizationResourceGenerator::GenerateLocResAndUpdateLiveEntriesFromConfig(InConfigFilePath, bSkipSourceCheck);
+		return FTextLocalizationResourceGenerator::GenerateLocResAndUpdateLiveEntriesFromConfig(InConfigFilePath, EGenerateLocResFlags::None);
 	}
 
 	virtual ULocalizationTarget* GetLocalizationTargetByName(FString TargetName, bool bIsEngineTarget) override
diff --git a/Engine/Source/Developer/Localization/Private/TextLocalizationResourceGenerator.cpp b/Engine/Source/Developer/Localization/Private/TextLocalizationResourceGenerator.cpp
--- a/Engine/Source/Developer/Localization/Private/TextLocalizationResourceGenerator.cpp
+++ b/Engine/Source/Developer/Localization/Private/TextLocalizationResourceGenerator.cpp
@@ -18,7 +18,7 @@ bool FTextLocalizationResourceGenerator::GenerateLocMeta(const FLocTextHelper& I
 	return true;
 }
 
-bool FTextLocalizationResourceGenerator::GenerateLocRes(const FLocTextHelper& InLocTextHelper, const FString& InCultureToGenerate, const bool bSkipSourceCheck, const FTextKey& InLocResID, FTextLocalizationResource& OutPlatformAgnosticLocRes, TMap<FName, TSharedRef<FTextLocalizationResource>>& OutPerPlatformLocRes, const int32 InPriority)
+bool FTextLocalizationResourceGenerator::GenerateLocRes(const FLocTextHelper& InLocTextHelper, const FString& InCultureToGenerate, const EGenerateLocResFlags InGenerateFlags, const FTextKey& InLocResID, FTextLocalizationResource& OutPlatformAgnosticLocRes, TMap<FName, TSharedRef<FTextLocalizationResource>>& OutPerPlatformLocRes, const int32 InPriority)
 {
 	const bool bIsNativeCulture = InCultureToGenerate == InLocTextHelper.GetNativeCulture();
 	FCulturePtr Culture = FInternationalization::Get().GetCulture(InCultureToGenerate);
@@ -34,21 +34,21 @@ bool FTextLocalizationResourceGenerator::GenerateLocRes(const FLocTextHelper& In
 	}
 
 	// Add each manifest entry to the LocRes file
-	InLocTextHelper.EnumerateSourceTexts([&InLocTextHelper, &InCultureToGenerate, &bSkipSourceCheck, &InLocResID, &OutPlatformAgnosticLocRes, &OutPerPlatformLocRes, InPriority, bIsNativeCulture, Culture](TSharedRef<FManifestEntry> InManifestEntry) -> bool
+	InLocTextHelper.EnumerateSourceTexts([&InLocTextHelper, &InCultureToGenerate, InGenerateFlags, &InLocResID, &OutPlatformAgnosticLocRes, &OutPerPlatformLocRes, InPriority, bIsNativeCulture, Culture](TSharedRef<FManifestEntry> InManifestEntry) -> bool
 	{
 		// For each context, we may need to create a different or even multiple LocRes entries.
 		for (const FManifestContext& Context : InManifestEntry->Contexts)
 		{
 			// Find the correct translation based upon the native source text
 			FLocItem TranslationText;
-			InLocTextHelper.GetRuntimeText(InCultureToGenerate, InManifestEntry->Namespace, Context.Key, Context.KeyMetadataObj, ELocTextExportSourceMethod::NativeText, InManifestEntry->Source, TranslationText, bSkipSourceCheck);
+			InLocTextHelper.GetRuntimeText(InCultureToGenerate, InManifestEntry->Namespace, Context.Key, Context.KeyMetadataObj, ELocTextExportSourceMethod::NativeText, InManifestEntry->Source, TranslationText, EnumHasAnyFlags(InGenerateFlags, EGenerateLocResFlags::AllowStaleTranslations));
 
 			// Is this entry considered translated? Native entries are always translated
 			const bool bIsTranslated = bIsNativeCulture || !InManifestEntry->Source.IsExactMatch(TranslationText);
 			if (bIsTranslated)
 			{
 				// Validate translations that look like they could be format patterns
-				if (Culture && TranslationText.Text.Contains(TEXT("{"), ESearchCase::CaseSensitive))
+				if (EnumHasAnyFlags(InGenerateFlags, EGenerateLocResFlags::ValidateFormatPatterns) && Culture && TranslationText.Text.Contains(TEXT("{"), ESearchCase::CaseSensitive))
 				{
 					const FTextFormat FmtPattern = FTextFormat::FromString(TranslationText.Text);
 
@@ -64,6 +64,21 @@ bool FTextLocalizationResourceGenerator::GenerateLocRes(const FLocTextHelper& In
 					}
 				}
 
+				// Validate that text doesn't have leading or trailing whitespace
+				if (EnumHasAnyFlags(InGenerateFlags, EGenerateLocResFlags::ValidateSafeWhitespace) && TranslationText.Text.Len() > 0)
+				{
+					auto IsUnsafeWhitespace = [](const TCHAR InChar)
+					{
+						// Unsafe whitespace is any whitespace character, except new-lines
+						return FText::IsWhitespace(InChar) && !(InChar == TEXT('\r') || InChar == TEXT('\n'));
+					};
+
+					if (IsUnsafeWhitespace(TranslationText.Text[0]) || IsUnsafeWhitespace(TranslationText.Text[TranslationText.Text.Len() - 1]))
+					{
+						UE_LOG(LogTextLocalizationResourceGenerator, Warning, TEXT("%s"), *FLocTextHelper::SanitizeLogOutput(FString::Printf(TEXT("Translation '%s' (%s,%s) has leading or trailing whitespace for '%s'"), *TranslationText.Text, *InManifestEntry->Namespace.GetString(), *Context.Key.GetString(), *InCultureToGenerate)));
+					}
+				}
+
 				// Find the LocRes to update
 				FTextLocalizationResource* LocResToUpdate = &OutPlatformAgnosticLocRes;
 				if (!Context.PlatformName.IsNone())
@@ -86,7 +101,7 @@ bool FTextLocalizationResourceGenerator::GenerateLocRes(const FLocTextHelper& In
 	return true;
 }
 
-bool FTextLocalizationResourceGenerator::GenerateLocResAndUpdateLiveEntriesFromConfig(const FString& InConfigFilePath, const bool bSkipSourceCheck)
+bool FTextLocalizationResourceGenerator::GenerateLocResAndUpdateLiveEntriesFromConfig(const FString& InConfigFilePath, const EGenerateLocResFlags InGenerateFlags)
 {
 	FInternationalization& I18N = FInternationalization::Get();
 
@@ -195,7 +210,7 @@ bool FTextLocalizationResourceGenerator::GenerateLocResAndUpdateLiveEntriesFromC
 		const FString CulturePath = DestinationPath / CultureName;
 		const FString ResourceFilePath = FPaths::ConvertRelativePathToFull(CulturePath / ResourceName);
 
-		if (!GenerateLocRes(LocTextHelper, CultureName, bSkipSourceCheck, FTextKey(ResourceFilePath), TextLocalizationResource, Unused_PerPlatformLocRes, CultureIndex))
+		if (!GenerateLocRes(LocTextHelper, CultureName, InGenerateFlags, FTextKey(ResourceFilePath), TextLocalizationResource, Unused_PerPlatformLocRes, CultureIndex))
 		{
 			UE_LOG(LogTextLocalizationResourceGenerator, Error, TEXT("Failed to generate localization resource for culture '%s'."), *CultureName);
 			return false;
diff --git a/Engine/Source/Developer/Localization/Public/LocalizationModule.h b/Engine/Source/Developer/Localization/Public/LocalizationModule.h
--- a/Engine/Source/Developer/Localization/Public/LocalizationModule.h
+++ b/Engine/Source/Developer/Localization/Public/LocalizationModule.h
@@ -14,7 +14,7 @@ public:
 	/**
 	 * Given a config file, generate a compiled LocRes file for the active culture and use it to update the live-entries in the localization manager.
 	 */
-	virtual bool HandleRegenLocCommand(const FString& InConfigFilePath, const bool bSkipSourceCheck) = 0;
+	virtual bool HandleRegenLocCommand(const FString& InConfigFilePath) = 0;
 
 	virtual ULocalizationTarget* GetLocalizationTargetByName(FString TargetName, bool bIsEngineTarget) = 0;
 
diff --git a/Engine/Source/Developer/Localization/Public/LocalizationTargetTypes.h b/Engine/Source/Developer/Localization/Public/LocalizationTargetTypes.h
--- a/Engine/Source/Developer/Localization/Public/LocalizationTargetTypes.h
+++ b/Engine/Source/Developer/Localization/Public/LocalizationTargetTypes.h
@@ -301,12 +301,22 @@ struct FLocalizationCompilationSettings
 
 	FLocalizationCompilationSettings()
 		: SkipSourceCheck(false)
+		, ValidateFormatPatterns(true)
+		, ValidateSafeWhitespace(false)
 	{
 	}
 
 	/* Should we skip the source check when compiling translations? This will allow translations whose source no longer matches the active source to still be used by the game at runtime. */
 	UPROPERTY(config, EditAnywhere, Category="Source")
 	bool SkipSourceCheck;
+
+	/* Should we validate that format patterns are valid for the culture being compiled (eg, detect invalid plural rules or broken syntax). */
+	UPROPERTY(config, EditAnywhere, Category="Validation")
+	bool ValidateFormatPatterns;
+
+	/* Should we validate that text doesn't contain any unsafe whitespace (leading or trailing whitespace) that could get lost during the translation process. */
+	UPROPERTY(config, EditAnywhere, Category="Validation")
+	bool ValidateSafeWhitespace;
 };
 
 USTRUCT()
diff --git a/Engine/Source/Developer/Localization/Public/TextLocalizationResourceGenerator.h b/Engine/Source/Developer/Localization/Public/TextLocalizationResourceGenerator.h
--- a/Engine/Source/Developer/Localization/Public/TextLocalizationResourceGenerator.h
+++ b/Engine/Source/Developer/Localization/Public/TextLocalizationResourceGenerator.h
@@ -3,12 +3,26 @@
 #pragma once 
 
 #include "CoreMinimal.h"
+#include "Misc/EnumClassFlags.h"
 #include "Internationalization/TextKey.h"
 
 class FLocTextHelper;
 class FTextLocalizationMetaDataResource;
 class FTextLocalizationResource;
 
+enum class EGenerateLocResFlags : uint8
+{
+	/** No special behavior */
+	None = 0,
+	/** Should "stale" translations (translations that were made against different source text) still be used? */
+	AllowStaleTranslations = 1<<0,
+	/** Validate that format patterns are valid for the culture being compiled (eg, detect invalid plural rules or broken syntax) */
+	ValidateFormatPatterns = 1<<1,
+	/** Validate that text doesn't contain any unsafe whitespace (leading or trailing whitespace) that could get lost during the translation process */
+	ValidateSafeWhitespace = 1<<2,
+};
+ENUM_CLASS_FLAGS(EGenerateLocResFlags);
+
 /** Utility functions for generating compiled LocMeta (Localization MetaData Resource) and LocRes (Localization Resource) files from source localization data */
 class FTextLocalizationResourceGenerator
 {
@@ -21,10 +35,10 @@ public:
 	/**
 	 * Given a loc text helper, generate a compiled LocRes resource for the given culture.
 	 */
-	LOCALIZATION_API static bool GenerateLocRes(const FLocTextHelper& InLocTextHelper, const FString& InCultureToGenerate, const bool bSkipSourceCheck, const FTextKey& InLocResID, FTextLocalizationResource& OutPlatformAgnosticLocRes, TMap<FName, TSharedRef<FTextLocalizationResource>>& OutPerPlatformLocRes, const int32 InPriority = 0);
+	LOCALIZATION_API static bool GenerateLocRes(const FLocTextHelper& InLocTextHelper, const FString& InCultureToGenerate, const EGenerateLocResFlags InGenerateFlags, const FTextKey& InLocResID, FTextLocalizationResource& OutPlatformAgnosticLocRes, TMap<FName, TSharedRef<FTextLocalizationResource>>& OutPerPlatformLocRes, const int32 InPriority = 0);
 
 	/**
 	 * Given a config file, generate a compiled LocRes resource for the active culture and use it to update the live-entries in the localization manager.
 	 */
-	LOCALIZATION_API static bool GenerateLocResAndUpdateLiveEntriesFromConfig(const FString& InConfigFilePath, const bool bSkipSourceCheck);
+	LOCALIZATION_API static bool GenerateLocResAndUpdateLiveEntriesFromConfig(const FString& InConfigFilePath, const EGenerateLocResFlags InGenerateFlags);
 };
diff --git a/Engine/Source/Editor/TranslationEditor/Private/TranslationDataManager.cpp b/Engine/Source/Editor/TranslationEditor/Private/TranslationDataManager.cpp
--- a/Engine/Source/Editor/TranslationEditor/Private/TranslationDataManager.cpp
+++ b/Engine/Source/Editor/TranslationEditor/Private/TranslationDataManager.cpp
@@ -641,7 +641,7 @@ void FTranslationDataManager::PreviewAllTranslationsInEditor(ULocalizationTarget
 	{
 		const FString ConfigFilePath = LocalizationConfigurationScript::GetRegenerateResourcesConfigPath(LocalizationTarget);
 		LocalizationConfigurationScript::GenerateRegenerateResourcesConfigFile(LocalizationTarget).Write(ConfigFilePath);
-		FTextLocalizationResourceGenerator::GenerateLocResAndUpdateLiveEntriesFromConfig(ConfigFilePath, /*bSkipSourceCheck*/false);
+		FTextLocalizationResourceGenerator::GenerateLocResAndUpdateLiveEntriesFromConfig(ConfigFilePath, EGenerateLocResFlags::None);
 	}
 	else
 	{
diff --git a/Engine/Source/Editor/UnrealEd/Private/Commandlets/GenerateTextLocalizationResourceCommandlet.cpp b/Engine/Source/Editor/UnrealEd/Private/Commandlets/GenerateTextLocalizationResourceCommandlet.cpp
--- a/Engine/Source/Editor/UnrealEd/Private/Commandlets/GenerateTextLocalizationResourceCommandlet.cpp
+++ b/Engine/Source/Editor/UnrealEd/Private/Commandlets/GenerateTextLocalizationResourceCommandlet.cpp
@@ -117,11 +117,27 @@ int32 UGenerateTextLocalizationResourceCommandlet::Main(const FString& Params)
 		return -1;
 	}
 
+	EGenerateLocResFlags GenerateFlags = EGenerateLocResFlags::None;
+
 	// Get whether to skip the source check.
-	bool bSkipSourceCheck = false;
-	if (!GetBoolFromConfig(*SectionName, TEXT("bSkipSourceCheck"), bSkipSourceCheck, GatherTextConfigPath))
 	{
-		bSkipSourceCheck = false;
+		bool bSkipSourceCheck = false;
+		GetBoolFromConfig(*SectionName, TEXT("bSkipSourceCheck"), bSkipSourceCheck, GatherTextConfigPath);
+		GenerateFlags |= (bSkipSourceCheck ? EGenerateLocResFlags::AllowStaleTranslations : EGenerateLocResFlags::None);
+	}
+
+	// Get whether to validate format patterns.
+	{
+		bool bValidateFormatPatterns = false;
+		GetBoolFromConfig(*SectionName, TEXT("bValidateFormatPatterns"), bValidateFormatPatterns, GatherTextConfigPath);
+		GenerateFlags |= (bValidateFormatPatterns ? EGenerateLocResFlags::ValidateFormatPatterns : EGenerateLocResFlags::None);
+	}
+
+	// Get whether to validate whitespace.
+	{
+		bool bValidateSafeWhitespace = false;
+		GetBoolFromConfig(*SectionName, TEXT("bValidateSafeWhitespace"), bValidateSafeWhitespace, GatherTextConfigPath);
+		GenerateFlags |= (bValidateSafeWhitespace ? EGenerateLocResFlags::ValidateSafeWhitespace : EGenerateLocResFlags::None);
 	}
 
 	// Load the manifest and all archives
@@ -184,7 +200,7 @@ int32 UGenerateTextLocalizationResourceCommandlet::Main(const FString& Params)
 		FTextLocalizationResource PlatformAgnosticLocRes;
 		TMap<FName, TSharedRef<FTextLocalizationResource>> PerPlatformLocRes;
 		const FTextKey LocResId = DestinationPath / CultureName / ResourceName;
-		if (!FTextLocalizationResourceGenerator::GenerateLocRes(LocTextHelper, CultureName, bSkipSourceCheck, LocResId, PlatformAgnosticLocRes, PerPlatformLocRes))
+		if (!FTextLocalizationResourceGenerator::GenerateLocRes(LocTextHelper, CultureName, GenerateFlags, LocResId, PlatformAgnosticLocRes, PerPlatformLocRes))
 		{
 			UE_LOG(LogGenerateTextLocalizationResourceCommandlet, Error, TEXT("Failed to generate LocRes %s"), LocResId.GetChars());
 			return false;
diff --git a/Engine/Source/Runtime/Engine/Private/UnrealEngine.cpp b/Engine/Source/Runtime/Engine/Private/UnrealEngine.cpp
--- a/Engine/Source/Runtime/Engine/Private/UnrealEngine.cpp
+++ b/Engine/Source/Runtime/Engine/Private/UnrealEngine.cpp
@@ -3714,7 +3714,7 @@ bool UEngine::Exec( UWorld* InWorld, const TCHAR* Cmd, FOutputDevice& Ar )
 		FString ConfigFilePath;
 		if (FParse::Value(Cmd, TEXT("REGENLOC="), ConfigFilePath))
 		{
-			ILocalizationModule::Get().HandleRegenLocCommand(ConfigFilePath, /*bSkipSourceCheck*/false);
+			ILocalizationModule::Get().HandleRegenLocCommand(ConfigFilePath);
 		}
 	}
 #endif
