Fixed uninitialized variables

#rnx

diff --git a/Engine/Source/Developer/Localization/Public/LocalizationTargetTypes.h b/Engine/Source/Developer/Localization/Public/LocalizationTargetTypes.h
--- a/Engine/Source/Developer/Localization/Public/LocalizationTargetTypes.h
+++ b/Engine/Source/Developer/Localization/Public/LocalizationTargetTypes.h
@@ -52,7 +52,7 @@ struct FGatherTextSearchDirectory
 	GENERATED_USTRUCT_BODY()
 
 	UPROPERTY(config, EditAnywhere, Category="Path")
-	ELocalizationGatherPathRoot PathRoot;
+	ELocalizationGatherPathRoot PathRoot = ELocalizationGatherPathRoot::Auto;
 
 	UPROPERTY(config, EditAnywhere, Category="Path")
 	FString Path;
@@ -66,7 +66,7 @@ struct FGatherTextIncludePath
 	GENERATED_USTRUCT_BODY()
 
 	UPROPERTY(config, EditAnywhere, Category="Path")
-	ELocalizationGatherPathRoot PathRoot;
+	ELocalizationGatherPathRoot PathRoot = ELocalizationGatherPathRoot::Auto;
 
 	UPROPERTY(config, EditAnywhere, Category="Pattern")
 	FString Pattern;
@@ -80,7 +80,7 @@ struct FGatherTextExcludePath
 	GENERATED_USTRUCT_BODY()
 
 	UPROPERTY(config, EditAnywhere, Category="Path")
-	ELocalizationGatherPathRoot PathRoot;
+	ELocalizationGatherPathRoot PathRoot = ELocalizationGatherPathRoot::Auto;
 
 	UPROPERTY(config, EditAnywhere, Category="Pattern")
 	FString Pattern;
