Fix Clang 8 compilation on Windows. Set default compilerto 8.0.0.

diff --git a/Engine/Source/Developer/Localization/Public/LocTextHelper.h b/Engine/Source/Developer/Localization/Public/LocTextHelper.h
--- a/Engine/Source/Developer/Localization/Public/LocTextHelper.h
+++ b/Engine/Source/Developer/Localization/Public/LocTextHelper.h
@@ -843,10 +843,6 @@ public:
 	 */
 	static bool FindKeysForLegacyTranslation(const TSharedRef<const FInternationalizationManifest>& InManifest, const TSharedPtr<const FInternationalizationArchive>& InNativeArchive, const FLocKey& InNamespace, const FString& InSource, const TSharedPtr<FLocMetadataObject> InKeyMetadataObj, TArray<FLocKey>& OutKeys);
 
-	/** Publicly movable */
-	FLocTextHelper(FLocTextHelper&&) = default;
-	FLocTextHelper& operator=(FLocTextHelper&&) = default;
-
 private:
 	/** Non-copyable */
 	FLocTextHelper(const FLocTextHelper&) = delete;
diff --git a/Engine/Source/Programs/UnrealBuildTool/Platform/Windows/UEBuildWindows.cs b/Engine/Source/Programs/UnrealBuildTool/Platform/Windows/UEBuildWindows.cs
--- a/Engine/Source/Programs/UnrealBuildTool/Platform/Windows/UEBuildWindows.cs
+++ b/Engine/Source/Programs/UnrealBuildTool/Platform/Windows/UEBuildWindows.cs
@@ -447,7 +447,7 @@ namespace UnrealBuildTool
 		/// <summary>
 		/// The default compiler version to be used, if installed. 
 		/// </summary>
-		static readonly VersionNumber DefaultClangToolChainVersion = VersionNumber.Parse("7.1.0");
+		static readonly VersionNumber DefaultClangToolChainVersion = VersionNumber.Parse("8.0.0");
 
 		/// <summary>
 		/// The default compiler version to be used, if installed. 
