Added option for confidental platforms to avoid splitting localization from restricted targets

#rnx

diff --git a/Engine/Source/Developer/Localization/Private/LocTextHelper.cpp b/Engine/Source/Developer/Localization/Private/LocTextHelper.cpp
--- a/Engine/Source/Developer/Localization/Private/LocTextHelper.cpp
+++ b/Engine/Source/Developer/Localization/Private/LocTextHelper.cpp
@@ -28,7 +28,25 @@ const TArray<FString>& FLocTextPlatformSplitUtils::GetPlatformsToSplit(const ELo
 	switch (InSplitMode)
 	{
 	case ELocTextPlatformSplitMode::Restricted:
-		return FDataDrivenPlatformInfoRegistry::GetConfidentialPlatforms();
+		{
+			static TArray<FString> RestrictedPlatformNames = []()
+			{
+				TArray<FString> TmpArray;
+				for (const FString& PlatformName : FDataDrivenPlatformInfoRegistry::GetConfidentialPlatforms())
+				{
+					const FDataDrivenPlatformInfoRegistry::FPlatformInfo& PlatformInfo = FDataDrivenPlatformInfoRegistry::GetPlatformInfo(PlatformName);
+					check(PlatformInfo.bIsConfidential);
+					if (PlatformInfo.bRestrictLocalization)
+					{
+						TmpArray.Add(PlatformName);
+					}
+				}
+				TmpArray.Sort();
+				return TmpArray;
+			}();
+			return RestrictedPlatformNames;
+		}
+		break;
 
 	case ELocTextPlatformSplitMode::All:
 		{
diff --git a/Engine/Source/Runtime/Core/Private/Misc/DataDrivenPlatformInfoRegistry.cpp b/Engine/Source/Runtime/Core/Private/Misc/DataDrivenPlatformInfoRegistry.cpp
--- a/Engine/Source/Runtime/Core/Private/Misc/DataDrivenPlatformInfoRegistry.cpp
+++ b/Engine/Source/Runtime/Core/Private/Misc/DataDrivenPlatformInfoRegistry.cpp
@@ -81,6 +81,7 @@ static const TMap<FString, FDataDrivenPlatformInfoRegistry::FPlatformInfo>& GetA
 				// cache info
 				FDataDrivenPlatformInfoRegistry::FPlatformInfo& Info = DataDrivenPlatforms.Add(PlatformName, FDataDrivenPlatformInfoRegistry::FPlatformInfo());
 				IniFile.GetBool(TEXT("DataDrivenPlatformInfo"), TEXT("bIsConfidential"), Info.bIsConfidential);
+				IniFile.GetBool(TEXT("DataDrivenPlatformInfo"), TEXT("bRestrictLocalization"), Info.bRestrictLocalization);
 
 				// get the parent to build list later
 				FString IniParent;
diff --git a/Engine/Source/Runtime/Core/Public/Misc/DataDrivenPlatformInfoRegistry.h b/Engine/Source/Runtime/Core/Public/Misc/DataDrivenPlatformInfoRegistry.h
--- a/Engine/Source/Runtime/Core/Public/Misc/DataDrivenPlatformInfoRegistry.h
+++ b/Engine/Source/Runtime/Core/Public/Misc/DataDrivenPlatformInfoRegistry.h
@@ -12,21 +12,18 @@
 struct CORE_API FDataDrivenPlatformInfoRegistry
 {
 	// Information about a platform loaded from disk
-	struct  FPlatformInfo
+	struct FPlatformInfo
 	{
-		FPlatformInfo()
-			: bIsConfidential(false)
-		{
-		}
-
 		// is this platform confidential
-		bool bIsConfidential;
+		bool bIsConfidential = false;
+
+		// should this platform be split when using ELocTextPlatformSplitMode::Restricted (only used when bIsConfidential is true)
+		bool bRestrictLocalization = false;
 
 		// cached list of ini parents
 		TArray<FString> IniParentChain;
 	};
 
-
 	/**
 	 * Get the data driven platform info for a given platform. If the platform doesn't have any on disk,
 	 * this will return a default constructed FConfigDataDrivenPlatformInfo
