Fixed how SaveSelectedTranslations matches translations and formats PO data

diff --git a/Engine/Source/Developer/Localization/Private/PortableObjectPipeline.cpp b/Engine/Source/Developer/Localization/Private/PortableObjectPipeline.cpp
--- a/Engine/Source/Developer/Localization/Private/PortableObjectPipeline.cpp
+++ b/Engine/Source/Developer/Localization/Private/PortableObjectPipeline.cpp
@@ -51,137 +51,6 @@ namespace
 		FLocKeyPairMultiMap CollapsedNSSourceStringToExpandedNSKey;
 	};
 
-	/**
-	* Declarations
-	*/
-	FString ConditionIdentityForPOMsgCtxt(const FString& Namespace, const FString& Key, const TSharedPtr<FLocMetadataObject>& KeyMetaData, const ELocalizedTextCollapseMode InTextCollapseMode);
-	void ParsePOMsgCtxtForIdentity(const FString& MsgCtxt, FString& OutNamespace, FString& OutKey);
-	FString ConditionArchiveStrForPo(const FString& InStr);
-	FString ConditionPoStringForArchive(const FString& InStr);
-	FString ConvertSrcLocationToPORef(const FString& InSrcLocation);
-	FString GetConditionedKeyForExtractedComment(const FString& Key);
-	FString GetConditionedReferenceForExtractedComment(const FString& PORefString);
-	FString GetConditionedInfoMetaDataForExtractedComment(const FString& KeyName, const FString& ValueString);
-	TSharedRef<FInternationalizationManifest> BuildCollapsedManifest(FLocTextHelper& InLocTextHelper, const ELocalizedTextCollapseMode InTextCollapseMode, FCollapsedData& OutCollapsedData);
-
-	/**
-	* Definitions
-	*/
-	FString ConditionIdentityForPOMsgCtxt(const FString& Namespace, const FString& Key, const TSharedPtr<FLocMetadataObject>& KeyMetaData, const ELocalizedTextCollapseMode InTextCollapseMode)
-	{
-		auto EscapeMsgCtxtParticleInline = [](FString& InStr)
-		{
-			InStr.ReplaceInline(TEXT(","), TEXT("\\,"), ESearchCase::CaseSensitive);
-		};
-
-		FString EscapedNamespace = Namespace;
-		EscapeMsgCtxtParticleInline(EscapedNamespace);
-
-		FString EscapedKey = Key;
-		EscapeMsgCtxtParticleInline(EscapedKey);
-
-		const bool bAppendKey = InTextCollapseMode != ELocalizedTextCollapseMode::IdenticalNamespaceAndSource || KeyMetaData.IsValid();
-		return ConditionArchiveStrForPo(bAppendKey ? FString::Printf(TEXT("%s,%s"), *EscapedNamespace, *EscapedKey) : EscapedNamespace);
-	}
-
-	void ParsePOMsgCtxtForIdentity(const FString& MsgCtxt, FString& OutNamespace, FString& OutKey)
-	{
-		auto UnescapeMsgCtxtParticleInline = [](FString& InStr)
-		{
-			InStr.ReplaceInline(TEXT("\\,"), TEXT(","), ESearchCase::CaseSensitive);
-		};
-
-		const FString ConditionedMsgCtxt = ConditionPoStringForArchive(MsgCtxt);
-
-		// Find the unescaped comma that defines the breaking point between the namespace and the key
-		int32 CommaIndex = INDEX_NONE;
-		{
-			bool bIsEscaped = false;
-			for (int32 Index = 0; Index < ConditionedMsgCtxt.Len(); ++Index)
-			{
-				if (bIsEscaped)
-				{
-					// No longer escaped, and skip this character
-					bIsEscaped = false;
-					continue;
-				}
-
-				if (ConditionedMsgCtxt[Index] == TEXT(','))
-				{
-					// Found the unescaped comma
-					CommaIndex = Index;
-					break;
-				}
-
-				if (ConditionedMsgCtxt[Index] == TEXT('\\'))
-				{
-					// Next character will be escaped
-					bIsEscaped = true;
-					continue;
-				}
-			}
-		}
-
-		if (CommaIndex == INDEX_NONE)
-		{
-			OutNamespace = ConditionedMsgCtxt;
-			OutKey.Reset();
-		}
-		else
-		{
-			OutNamespace = ConditionedMsgCtxt.Mid(0, CommaIndex);
-			OutKey = ConditionedMsgCtxt.Mid(CommaIndex + 1);
-		}
-
-		UnescapeMsgCtxtParticleInline(OutNamespace);
-		UnescapeMsgCtxtParticleInline(OutKey);
-	}
-
-	FString ConditionArchiveStrForPo(const FString& InStr)
-	{
-		FString Result = InStr;
-		Result.ReplaceInline(TEXT("\\"), TEXT("\\\\"), ESearchCase::CaseSensitive);
-		Result.ReplaceInline(TEXT("\""), TEXT("\\\""), ESearchCase::CaseSensitive);
-		Result.ReplaceInline(TEXT("\r"), TEXT("\\r"), ESearchCase::CaseSensitive);
-		Result.ReplaceInline(TEXT("\n"), TEXT("\\n"), ESearchCase::CaseSensitive);
-		Result.ReplaceInline(TEXT("\t"), TEXT("\\t"), ESearchCase::CaseSensitive);
-		return Result;
-	}
-
-	FString ConditionPoStringForArchive(const FString& InStr)
-	{
-		FString Result = InStr;
-		Result.ReplaceInline(TEXT("\\t"), TEXT("\t"), ESearchCase::CaseSensitive);
-		Result.ReplaceInline(TEXT("\\n"), TEXT("\n"), ESearchCase::CaseSensitive);
-		Result.ReplaceInline(TEXT("\\r"), TEXT("\r"), ESearchCase::CaseSensitive);
-		Result.ReplaceInline(TEXT("\\\""), TEXT("\""), ESearchCase::CaseSensitive);
-		Result.ReplaceInline(TEXT("\\\\"), TEXT("\\"), ESearchCase::CaseSensitive);
-		return Result;
-	}
-
-	FString ConvertSrcLocationToPORef(const FString& InSrcLocation)
-	{
-		// Source location format: /Path1/Path2/file.cpp - line 123
-		// PO Reference format: /Path1/Path2/file.cpp:123
-		// @TODO: Note, we assume the source location format here but it could be arbitrary.
-		return InSrcLocation.Replace(TEXT(" - line "), TEXT(":"), ESearchCase::CaseSensitive);
-	}
-
-	FString GetConditionedKeyForExtractedComment(const FString& Key)
-	{
-		return FString::Printf(TEXT("Key:\t%s"), *Key);
-	}
-
-	FString GetConditionedReferenceForExtractedComment(const FString& PORefString)
-	{
-		return FString::Printf(TEXT("SourceLocation:\t%s"), *PORefString);
-	}
-
-	FString GetConditionedInfoMetaDataForExtractedComment(const FString& KeyName, const FString& ValueString)
-	{
-		return FString::Printf(TEXT("InfoMetaData:\t\"%s\" : \"%s\""), *KeyName, *ValueString);
-	}
-
 	TSharedRef<FInternationalizationManifest> BuildCollapsedManifest(FLocTextHelper& InLocTextHelper, const ELocalizedTextCollapseMode InTextCollapseMode, FCollapsedData& OutCollapsedData)
 	{
 		TSharedRef<FInternationalizationManifest> CollapsedManifest = MakeShared<FInternationalizationManifest>();
@@ -319,6 +188,8 @@ namespace
 
 	bool ImportPortableObject(FLocTextHelper& InLocTextHelper, const FString& InCulture, const FString& InPOFilePath, const FCollapsedData& InCollapsedData)
 	{
+		using namespace PortableObjectPipeline;
+
 		if (!FPaths::FileExists(InPOFilePath))
 		{
 			UE_LOG(LogPortableObjectPipeline, Warning, TEXT("Could not find file %s"), *InPOFilePath);
@@ -425,6 +296,8 @@ namespace
 
 	bool ExportPortableObject(FLocTextHelper& InLocTextHelper, const FString& InCulture, const FString& InPOFilePath, const ELocalizedTextCollapseMode InTextCollapseMode, TSharedRef<FInternationalizationManifest> InCollapsedManifest, const FCollapsedData& InCollapsedData, const bool bShouldPersistComments)
 	{
+		using namespace PortableObjectPipeline;
+
 		FPortableObjectFormatDOM NewPortableObject;
 
 		FString LocLang;
@@ -634,3 +507,118 @@ bool PortableObjectPipeline::ExportAll(FLocTextHelper& InLocTextHelper, const FS
 
 	return bSuccess;
 }
+
+FString PortableObjectPipeline::ConditionIdentityForPOMsgCtxt(const FString& Namespace, const FString& Key, const TSharedPtr<FLocMetadataObject>& KeyMetaData, const ELocalizedTextCollapseMode InTextCollapseMode)
+{
+	auto EscapeMsgCtxtParticleInline = [](FString& InStr)
+	{
+		InStr.ReplaceInline(TEXT(","), TEXT("\\,"), ESearchCase::CaseSensitive);
+	};
+
+	FString EscapedNamespace = Namespace;
+	EscapeMsgCtxtParticleInline(EscapedNamespace);
+
+	FString EscapedKey = Key;
+	EscapeMsgCtxtParticleInline(EscapedKey);
+
+	const bool bAppendKey = InTextCollapseMode != ELocalizedTextCollapseMode::IdenticalNamespaceAndSource || KeyMetaData.IsValid();
+	return ConditionArchiveStrForPo(bAppendKey ? FString::Printf(TEXT("%s,%s"), *EscapedNamespace, *EscapedKey) : EscapedNamespace);
+}
+
+void PortableObjectPipeline::ParsePOMsgCtxtForIdentity(const FString& MsgCtxt, FString& OutNamespace, FString& OutKey)
+{
+	auto UnescapeMsgCtxtParticleInline = [](FString& InStr)
+	{
+		InStr.ReplaceInline(TEXT("\\,"), TEXT(","), ESearchCase::CaseSensitive);
+	};
+
+	const FString ConditionedMsgCtxt = ConditionPoStringForArchive(MsgCtxt);
+
+	// Find the unescaped comma that defines the breaking point between the namespace and the key
+	int32 CommaIndex = INDEX_NONE;
+	{
+		bool bIsEscaped = false;
+		for (int32 Index = 0; Index < ConditionedMsgCtxt.Len(); ++Index)
+		{
+			if (bIsEscaped)
+			{
+				// No longer escaped, and skip this character
+				bIsEscaped = false;
+				continue;
+			}
+
+			if (ConditionedMsgCtxt[Index] == TEXT(','))
+			{
+				// Found the unescaped comma
+				CommaIndex = Index;
+				break;
+			}
+
+			if (ConditionedMsgCtxt[Index] == TEXT('\\'))
+			{
+				// Next character will be escaped
+				bIsEscaped = true;
+				continue;
+			}
+		}
+	}
+
+	if (CommaIndex == INDEX_NONE)
+	{
+		OutNamespace = ConditionedMsgCtxt;
+		OutKey.Reset();
+	}
+	else
+	{
+		OutNamespace = ConditionedMsgCtxt.Mid(0, CommaIndex);
+		OutKey = ConditionedMsgCtxt.Mid(CommaIndex + 1);
+	}
+
+	UnescapeMsgCtxtParticleInline(OutNamespace);
+	UnescapeMsgCtxtParticleInline(OutKey);
+}
+
+FString PortableObjectPipeline::ConditionArchiveStrForPo(const FString& InStr)
+{
+	FString Result = InStr;
+	Result.ReplaceInline(TEXT("\\"), TEXT("\\\\"), ESearchCase::CaseSensitive);
+	Result.ReplaceInline(TEXT("\""), TEXT("\\\""), ESearchCase::CaseSensitive);
+	Result.ReplaceInline(TEXT("\r"), TEXT("\\r"), ESearchCase::CaseSensitive);
+	Result.ReplaceInline(TEXT("\n"), TEXT("\\n"), ESearchCase::CaseSensitive);
+	Result.ReplaceInline(TEXT("\t"), TEXT("\\t"), ESearchCase::CaseSensitive);
+	return Result;
+}
+
+FString PortableObjectPipeline::ConditionPoStringForArchive(const FString& InStr)
+{
+	FString Result = InStr;
+	Result.ReplaceInline(TEXT("\\t"), TEXT("\t"), ESearchCase::CaseSensitive);
+	Result.ReplaceInline(TEXT("\\n"), TEXT("\n"), ESearchCase::CaseSensitive);
+	Result.ReplaceInline(TEXT("\\r"), TEXT("\r"), ESearchCase::CaseSensitive);
+	Result.ReplaceInline(TEXT("\\\""), TEXT("\""), ESearchCase::CaseSensitive);
+	Result.ReplaceInline(TEXT("\\\\"), TEXT("\\"), ESearchCase::CaseSensitive);
+	return Result;
+}
+
+FString PortableObjectPipeline::ConvertSrcLocationToPORef(const FString& InSrcLocation)
+{
+	// Source location format: /Path1/Path2/file.cpp - line 123
+	// PO Reference format: /Path1/Path2/file.cpp:123
+	// @TODO: Note, we assume the source location format here but it could be arbitrary.
+	return InSrcLocation.Replace(TEXT(" - line "), TEXT(":"), ESearchCase::CaseSensitive);
+}
+
+FString PortableObjectPipeline::GetConditionedKeyForExtractedComment(const FString& Key)
+{
+	return FString::Printf(TEXT("Key:\t%s"), *Key);
+}
+
+FString PortableObjectPipeline::GetConditionedReferenceForExtractedComment(const FString& PORefString)
+{
+	return FString::Printf(TEXT("SourceLocation:\t%s"), *PORefString);
+}
+
+FString PortableObjectPipeline::GetConditionedInfoMetaDataForExtractedComment(const FString& KeyName, const FString& ValueString)
+{
+	return FString::Printf(TEXT("InfoMetaData:\t\"%s\" : \"%s\""), *KeyName, *ValueString);
+}
diff --git a/Engine/Source/Developer/Localization/Public/PortableObjectPipeline.h b/Engine/Source/Developer/Localization/Public/PortableObjectPipeline.h
--- a/Engine/Source/Developer/Localization/Public/PortableObjectPipeline.h
+++ b/Engine/Source/Developer/Localization/Public/PortableObjectPipeline.h
@@ -3,10 +3,10 @@
 #pragma once 
 
 #include "CoreMinimal.h"
-#include "UObject/ObjectMacros.h"
 #include "PortableObjectPipeline.generated.h"
 
 class FLocTextHelper;
+class FLocMetadataObject;
 
 UENUM()
 enum class ELocalizedTextCollapseMode : uint8
@@ -32,4 +32,28 @@ namespace PortableObjectPipeline
 
 	/** Use the given LocTextHelper to generate a new PO file using the translation data for all cultures */
 	LOCALIZATION_API bool ExportAll(FLocTextHelper& InLocTextHelper, const FString& InPOCultureRootPath, const FString& InPOFilename, const ELocalizedTextCollapseMode InTextCollapseMode, const bool bShouldPersistComments, const bool bUseCultureDirectory);
+
+	/** Given a namespace and key, condition this information so it can be written to the "msgctxt" field of a PO file */
+	LOCALIZATION_API FString ConditionIdentityForPOMsgCtxt(const FString& Namespace, const FString& Key, const TSharedPtr<FLocMetadataObject>& KeyMetaData, const ELocalizedTextCollapseMode InTextCollapseMode);
+
+	/** Given the "msgctxt" field of a PO file, split it into the namespace and key */
+	LOCALIZATION_API void ParsePOMsgCtxtForIdentity(const FString& MsgCtxt, FString& OutNamespace, FString& OutKey);
+
+	/** Given a string, condition it so it can be written as a field of a PO file */
+	LOCALIZATION_API FString ConditionArchiveStrForPo(const FString& InStr);
+
+	/** Given the field of a PO file, condition it back to a clean string */
+	LOCALIZATION_API FString ConditionPoStringForArchive(const FString& InStr);
+
+	/** Given a source location string, clean it up for use in PO file fields */
+	LOCALIZATION_API FString ConvertSrcLocationToPORef(const FString& InSrcLocation);
+
+	/** Given a key string, condition it so it can be written as the extracted comment field of a PO file */
+	LOCALIZATION_API FString GetConditionedKeyForExtractedComment(const FString& Key);
+
+	/** Given a source location string, condition it so it can be written as the extracted comment field of a PO file */
+	LOCALIZATION_API FString GetConditionedReferenceForExtractedComment(const FString& PORefString);
+
+	/** Given a meta-data value string, condition it so it can be written as the extracted comment field of a PO file */
+	LOCALIZATION_API FString GetConditionedInfoMetaDataForExtractedComment(const FString& KeyName, const FString& ValueString);
 }
diff --git a/Engine/Source/Editor/TranslationEditor/Private/TranslationDataManager.cpp b/Engine/Source/Editor/TranslationEditor/Private/TranslationDataManager.cpp
--- a/Engine/Source/Editor/TranslationEditor/Private/TranslationDataManager.cpp
+++ b/Engine/Source/Editor/TranslationEditor/Private/TranslationDataManager.cpp
@@ -1,6 +1,7 @@
 // Copyright 1998-2018 Epic Games, Inc. All Rights Reserved.
 
 #include "TranslationDataManager.h"
+#include "PortableObjectPipeline.h"
 #include "Internationalization/InternationalizationManifest.h"
 #include "Internationalization/InternationalizationArchive.h"
 #include "Misc/FileHelper.h"
@@ -928,27 +929,44 @@ bool FTranslationDataManager::SaveSelectedTranslations(TArray<UTranslationUnit*>
 					// Search all translations for the one that matches this FText
 					for (UTranslationUnit* Translation : TranslationsArray)
 					{
+						bool bFoundMatchingTranslation = false;
+
 						// If namespace matches...
 						if (Translation->Namespace.Equals(EditedItem->Namespace))
 						{
-							// And source matches
-							if (Translation->Source.Equals(EditedItem->Source))
+							// And key matches...
+							for (const FTranslationContextInfo& ContextInfo : Translation->Contexts)
 							{
-								// Update the translation in TranslationDataManager, and finish searching these translations
-								Translation->Translation = EditedItem->Translation;
-
-								TSharedRef<FPortableObjectEntry> NewEntry = MakeShareable(new FPortableObjectEntry);
-								for (FTranslationContextInfo ContextInfo : Translation->Contexts)
+								if (ContextInfo.Key.Equals(EditedItem->Key))
 								{
-									NewEntry->ExtractedComments.Add(ContextInfo.Key);
-									NewEntry->ReferenceComments.Add(ContextInfo.Context);
-								}
+									bFoundMatchingTranslation = true;
+
+									// Update the translation in TranslationDataManager, and finish searching these translations
+									Translation->Translation = EditedItem->Translation;
+
+									// Add the PO entry
+									{
+										TSharedRef<FPortableObjectEntry> PoEntry = MakeShareable(new FPortableObjectEntry());
+
+										PoEntry->MsgId = PortableObjectPipeline::ConditionArchiveStrForPo(Translation->Source);
+										PoEntry->MsgCtxt = PortableObjectPipeline::ConditionIdentityForPOMsgCtxt(Translation->Namespace, ContextInfo.Key, nullptr, LocalizationTarget->Settings.ExportSettings.CollapseMode);
+										PoEntry->MsgStr.Add(PortableObjectPipeline::ConditionArchiveStrForPo(Translation->Translation));
 
-								NewEntry->MsgCtxt = Translation->Namespace;
-								NewEntry->MsgId = Translation->Source;
-								NewEntry->MsgStr.Add(Translation->Translation);
-								PortableObjectDom.AddEntry(NewEntry);
+										//@TODO: We support additional metadata entries that can be translated.  How do those fit in the PO file format?  Ex: isMature
+										const FString PORefString = PortableObjectPipeline::ConvertSrcLocationToPORef(ContextInfo.Context);
+										PoEntry->AddReference(PORefString); // Source location.
 
+										PoEntry->AddExtractedComment(PortableObjectPipeline::GetConditionedKeyForExtractedComment(ContextInfo.Key)); // "Notes from Programmer" in the form of the Key.
+										PoEntry->AddExtractedComment(PortableObjectPipeline::GetConditionedReferenceForExtractedComment(PORefString)); // "Notes from Programmer" in the form of the Source Location, since this comes in handy too and OneSky doesn't properly show references, only comments.
+
+										PortableObjectDom.AddEntry(PoEntry);
+									}
+									break;
+								}
+							}
+
+							if (bFoundMatchingTranslation)
+							{
 								break;
 							}
 						}
diff --git a/Engine/Source/Editor/TranslationEditor/Private/TranslationPickerEditWindow.cpp b/Engine/Source/Editor/TranslationEditor/Private/TranslationPickerEditWindow.cpp
--- a/Engine/Source/Editor/TranslationEditor/Private/TranslationPickerEditWindow.cpp
+++ b/Engine/Source/Editor/TranslationEditor/Private/TranslationPickerEditWindow.cpp
@@ -269,6 +269,7 @@ void STranslationPickerEditWidget::Construct(const FArguments& InArgs)
 	// Save the necessary data in UTranslationUnit for later.  This is what we pass to TranslationDataManager to save our edits
 	TranslationUnit = NewObject<UTranslationUnit>();
 	TranslationUnit->Namespace = CleanNamespaceString;
+	TranslationUnit->Key = KeyString.Get(TEXT(""));
 	TranslationUnit->Source = SourceString != nullptr ? *SourceString : TEXT("");
 	TranslationUnit->Translation = TranslationString;
 	TranslationUnit->LocresPath = LocresFullPath;
