Fixed empty translations being added to non-native LocRes files

New text always has an empty translations by default, so this was causing new text to go missing for non-native languages

diff --git a/Engine/Source/Developer/Localization/Private/LocTextHelper.cpp b/Engine/Source/Developer/Localization/Private/LocTextHelper.cpp
--- a/Engine/Source/Developer/Localization/Private/LocTextHelper.cpp
+++ b/Engine/Source/Developer/Localization/Private/LocTextHelper.cpp
@@ -973,7 +973,7 @@ void FLocTextHelper::GetRuntimeText(const FString& InCulture, const FLocKey& InN
 	OutTranslation = InSource;
 
 	TSharedPtr<FArchiveEntry> ArchiveEntry = FindTranslationImpl(InCulture, InNamespace, InKey, InKeyMetadataObj);
-	if (ArchiveEntry.IsValid())
+	if (ArchiveEntry.IsValid() && !ArchiveEntry->Translation.Text.IsEmpty())
 	{
 		if (bSkipSourceCheck)
 		{
