Added support for splitting off per-platform localization

This is handled by the ELocTextPlatformSplitMode enum, used by FLocTextHelper (default is ELocTextPlatformSplitMode::None to preserve old behavior), and allows you to optionally split localization data for all platforms, or for any NDA restricted platforms. Note: Any localization data that is used by multiple platforms is still considered platform agnostic, even if it appears only within platform specific files.

diff --git a/Engine/Config/Localization/Category.ini b/Engine/Config/Localization/Category.ini
--- a/Engine/Config/Localization/Category.ini
+++ b/Engine/Config/Localization/Category.ini
@@ -4,6 +4,7 @@ DestinationPath=./Content/Localization/Category
 ManifestName=Category.manifest
 ArchiveName=Category.archive
 PortableObjectName=Category.po
+PlatformSplitMode=ELocTextPlatformSplitMode::Restricted
 NativeCulture=en
 CulturesToGenerate=ar
 CulturesToGenerate=es
diff --git a/Engine/Config/Localization/Editor.ini b/Engine/Config/Localization/Editor.ini
--- a/Engine/Config/Localization/Editor.ini
+++ b/Engine/Config/Localization/Editor.ini
@@ -5,6 +5,7 @@ DestinationPath=Content/Localization/Editor
 ManifestName=Editor.manifest
 ArchiveName=Editor.archive
 PortableObjectName=Editor.po
+PlatformSplitMode=ELocTextPlatformSplitMode::Restricted
 NativeCulture=en
 CulturesToGenerate=ar
 CulturesToGenerate=en
diff --git a/Engine/Config/Localization/EditorTutorials.ini b/Engine/Config/Localization/EditorTutorials.ini
--- a/Engine/Config/Localization/EditorTutorials.ini
+++ b/Engine/Config/Localization/EditorTutorials.ini
@@ -10,6 +10,7 @@ DestinationPath=Content/Localization/EditorTutorials
 ManifestName=EditorTutorials.manifest
 ArchiveName=EditorTutorials.archive
 PortableObjectName=EditorTutorials.po
+PlatformSplitMode=ELocTextPlatformSplitMode::Restricted
 NativeCulture=en
 CulturesToGenerate=ar
 CulturesToGenerate=es
diff --git a/Engine/Config/Localization/Engine.ini b/Engine/Config/Localization/Engine.ini
--- a/Engine/Config/Localization/Engine.ini
+++ b/Engine/Config/Localization/Engine.ini
@@ -4,6 +4,7 @@ DestinationPath=Content/Localization/Engine
 ManifestName=Engine.manifest
 ArchiveName=Engine.archive
 PortableObjectName=Engine.po
+PlatformSplitMode=ELocTextPlatformSplitMode::Restricted
 NativeCulture=en
 CulturesToGenerate=ar
 CulturesToGenerate=en
diff --git a/Engine/Config/Localization/Keywords.ini b/Engine/Config/Localization/Keywords.ini
--- a/Engine/Config/Localization/Keywords.ini
+++ b/Engine/Config/Localization/Keywords.ini
@@ -4,6 +4,7 @@ DestinationPath=./Content/Localization/Keywords
 ManifestName=Keywords.manifest
 ArchiveName=Keywords.archive
 PortableObjectName=Keywords.po
+PlatformSplitMode=ELocTextPlatformSplitMode::Restricted
 NativeCulture=en
 CulturesToGenerate=ar
 CulturesToGenerate=es
diff --git a/Engine/Config/Localization/PropertyNames.ini b/Engine/Config/Localization/PropertyNames.ini
--- a/Engine/Config/Localization/PropertyNames.ini
+++ b/Engine/Config/Localization/PropertyNames.ini
@@ -4,6 +4,7 @@ DestinationPath=Content/Localization/PropertyNames
 ManifestName=PropertyNames.manifest
 ArchiveName=PropertyNames.archive
 PortableObjectName=PropertyNames.po
+PlatformSplitMode=ELocTextPlatformSplitMode::Restricted
 NativeCulture=en
 CulturesToGenerate=ar
 CulturesToGenerate=es
diff --git a/Engine/Config/Localization/ToolTips.ini b/Engine/Config/Localization/ToolTips.ini
--- a/Engine/Config/Localization/ToolTips.ini
+++ b/Engine/Config/Localization/ToolTips.ini
@@ -5,6 +5,7 @@ DestinationPath=Content/Localization/ToolTips
 ManifestName=ToolTips.manifest
 ArchiveName=ToolTips.archive
 PortableObjectName=ToolTips.po
+PlatformSplitMode=ELocTextPlatformSplitMode::Restricted
 NativeCulture=en
 CulturesToGenerate=ar
 CulturesToGenerate=es
diff --git a/Engine/Source/Developer/Localization/Localization.Build.cs b/Engine/Source/Developer/Localization/Localization.Build.cs
--- a/Engine/Source/Developer/Localization/Localization.Build.cs
+++ b/Engine/Source/Developer/Localization/Localization.Build.cs
@@ -17,6 +17,7 @@ public class Localization : ModuleRules
 		PrivateDependencyModuleNames.AddRange(
 			new string[]
 			{
+				"DesktopPlatform",
 				"Slate",
 				"SlateCore",
 				"InputCore",
diff --git a/Engine/Source/Developer/Localization/Private/LocTextHelper.cpp b/Engine/Source/Developer/Localization/Private/LocTextHelper.cpp
--- a/Engine/Source/Developer/Localization/Private/LocTextHelper.cpp
+++ b/Engine/Source/Developer/Localization/Private/LocTextHelper.cpp
@@ -1,8 +1,10 @@
 // Copyright 1998-2019 Epic Games, Inc. All Rights Reserved.
 
 #include "LocTextHelper.h"
+#include "PlatformInfo.h"
 #include "Misc/Paths.h"
 #include "Misc/FileHelper.h"
+#include "Misc/DataDrivenPlatformInfoRegistry.h"
 #include "Serialization/Csv/CsvParser.h"
 #include "Internationalization/IBreakIterator.h"
 #include "Internationalization/BreakIterator.h"
@@ -16,12 +18,52 @@
 
 DEFINE_LOG_CATEGORY_STATIC(LogLocTextHelper, Log, All);
 
+bool FLocTextPlatformSplitUtils::ShouldSplitPlatformData(const ELocTextPlatformSplitMode& InSplitMode)
+{
+	return InSplitMode != ELocTextPlatformSplitMode::None;
+}
+
+const TArray<FString>& FLocTextPlatformSplitUtils::GetPlatformsToSplit(const ELocTextPlatformSplitMode& InSplitMode)
+{
+	switch (InSplitMode)
+	{
+	case ELocTextPlatformSplitMode::Restricted:
+		return FDataDrivenPlatformInfoRegistry::GetConfidentialPlatforms();
+
+	case ELocTextPlatformSplitMode::All:
+		{
+			static TArray<FString> AllPlatformNames = []()
+			{
+				TArray<FString> TmpArray;
+				for (const PlatformInfo::FPlatformInfo& Info : PlatformInfo::EnumeratePlatformInfoArray(false))
+				{
+					if (!Info.IniPlatformName.IsEmpty())
+					{
+						TmpArray.AddUnique(*Info.IniPlatformName);
+					}
+				}
+				TmpArray.Sort();
+				return TmpArray;
+			}();
+			return AllPlatformNames;
+		}
+		break;
+
+	default:
+		break;
+	}
+
+	static TArray<FString> EmptyArray;
+	return EmptyArray;
+}
+
+
 void FLocTextConflicts::AddConflict(const FLocKey& InNamespace, const FLocKey& InKey, const TSharedPtr<FLocMetadataObject>& InKeyMetadata, const FLocItem& InSource, const FString& InSourceLocation)
 {
 	TSharedPtr<FConflict> ExistingEntry = FindEntryByKey(InNamespace, InKey, InKeyMetadata);
 	if (!ExistingEntry.IsValid())
 	{
-		TSharedRef<FConflict> NewEntry = MakeShareable(new FConflict(InNamespace, InKey, InKeyMetadata));
+		TSharedRef<FConflict> NewEntry = MakeShared<FConflict>(InNamespace, InKey, InKeyMetadata);
 		EntriesByKey.Add(InKey, NewEntry);
 		ExistingEntry = NewEntry;
 	}
@@ -326,13 +368,15 @@ void FLocTextWordCounts::SortRowsByDate()
 }
 
 
-FLocTextHelper::FLocTextHelper(TSharedPtr<ILocFileNotifies> InLocFileNotifies)
-	: LocFileNotifies(MoveTemp(InLocFileNotifies))
+FLocTextHelper::FLocTextHelper(TSharedPtr<ILocFileNotifies> InLocFileNotifies, const ELocTextPlatformSplitMode InPlatformSplitMode)
+	: PlatformSplitMode(InPlatformSplitMode)
+	, LocFileNotifies(MoveTemp(InLocFileNotifies))
 {
 }
 
-FLocTextHelper::FLocTextHelper(FString InTargetPath, FString InManifestName, FString InArchiveName, FString InNativeCulture, TArray<FString> InForeignCultures, TSharedPtr<ILocFileNotifies> InLocFileNotifies)
-	: TargetPath(MoveTemp(InTargetPath))
+FLocTextHelper::FLocTextHelper(FString InTargetPath, FString InManifestName, FString InArchiveName, FString InNativeCulture, TArray<FString> InForeignCultures, TSharedPtr<ILocFileNotifies> InLocFileNotifies, const ELocTextPlatformSplitMode InPlatformSplitMode)
+	: PlatformSplitMode(InPlatformSplitMode)
+	, TargetPath(MoveTemp(InTargetPath))
 	, ManifestName(MoveTemp(InManifestName))
 	, ArchiveName(MoveTemp(InArchiveName))
 	, NativeCulture(MoveTemp(InNativeCulture))
@@ -353,6 +397,21 @@ FLocTextHelper::FLocTextHelper(FString InTargetPath, FString InManifestName, FSt
 	}
 }
 
+bool FLocTextHelper::ShouldSplitPlatformData() const
+{
+	return FLocTextPlatformSplitUtils::ShouldSplitPlatformData(PlatformSplitMode);
+}
+
+ELocTextPlatformSplitMode FLocTextHelper::GetPlatformSplitMode() const
+{
+	return PlatformSplitMode;
+}
+
+const TArray<FString>& FLocTextHelper::GetPlatformsToSplit() const
+{
+	return FLocTextPlatformSplitUtils::GetPlatformsToSplit(PlatformSplitMode);
+}
+
 const FString& FLocTextHelper::GetTargetName() const
 {
 	return TargetName;
@@ -436,7 +495,7 @@ void FLocTextHelper::TrimManifest()
 	if (Dependencies.Num() > 0)
 	{
 		// We'll generate a new manifest by only including items that are not in the dependencies
-		TSharedRef<FInternationalizationManifest> TrimmedManifest = MakeShareable(new FInternationalizationManifest());
+		TSharedRef<FInternationalizationManifest> TrimmedManifest = MakeShared<FInternationalizationManifest>();
 
 		for (FManifestEntryByStringContainer::TConstIterator It(Manifest->GetEntriesBySourceTextIterator()); It; ++It)
 		{
@@ -445,7 +504,18 @@ void FLocTextHelper::TrimManifest()
 			for (const FManifestContext& Context : ManifestEntry->Contexts)
 			{
 				FString DependencyFileName;
-				const TSharedPtr<FManifestEntry> DependencyEntry = FindDependencyEntry(ManifestEntry->Namespace, Context, &DependencyFileName);
+				TSharedPtr<FManifestEntry> DependencyEntry = FindDependencyEntry(ManifestEntry->Namespace, Context, &DependencyFileName);
+
+				// Ignore this dependency if the platforms are different
+				if (DependencyEntry.IsValid())
+				{
+					const FManifestContext* DependencyContext = DependencyEntry->FindContext(Context.Key, Context.KeyMetadataObj);
+					if (Context.PlatformName != DependencyContext->PlatformName)
+					{
+						DependencyEntry.Reset();
+						DependencyFileName.Reset();
+					}
+				}
 
 				if (DependencyEntry.IsValid())
 				{
@@ -656,7 +726,7 @@ void FLocTextHelper::TrimArchive(const FString& InCulture)
 	}
 
 	// Copy any translations that match current manifest entries over into the trimmed archive
-	TSharedRef<FInternationalizationArchive> TrimmedArchive = MakeShareable(new FInternationalizationArchive());
+	TSharedRef<FInternationalizationArchive> TrimmedArchive = MakeShared<FInternationalizationArchive>();
 	EnumerateSourceTexts([&](TSharedRef<FManifestEntry> InManifestEntry) -> bool
 	{
 		for (const FManifestContext& Context : InManifestEntry->Contexts)
@@ -766,6 +836,17 @@ bool FLocTextHelper::AddSourceText(const FLocKey& InNamespace, const FLocItem& I
 	if (!ExistingEntry.IsValid())
 	{
 		ExistingEntry = FindDependencyEntry(InNamespace, InContext, &ExistingEntryFileName);
+
+		// Ignore this dependency if the platforms are different
+		if (ExistingEntry.IsValid())
+		{
+			const FManifestContext* DependencyContext = ExistingEntry->FindContext(InContext.Key, InContext.KeyMetadataObj);
+			if (InContext.PlatformName != DependencyContext->PlatformName)
+			{
+				ExistingEntry.Reset();
+				ExistingEntryFileName.Reset();
+			}
+		}
 	}
 
 	if (ExistingEntry.IsValid())
@@ -773,12 +854,13 @@ bool FLocTextHelper::AddSourceText(const FLocKey& InNamespace, const FLocItem& I
 		if (InSource.IsExactMatch(ExistingEntry->Source))
 		{
 			bAddSuccessful = true;
+			ExistingEntry->MergeContextPlatformInfo(InContext);
 		}
 		else
 		{
 			// Grab the source location of the conflicting context
 			const FManifestContext* ConflictingContext = ExistingEntry->FindContext(InContext.Key, InContext.KeyMetadataObj);
-			const FString ExistingEntrySourceLocation = (!ExistingEntryFileName.IsEmpty()) ? ExistingEntryFileName : ConflictingContext->SourceLocation;
+			const FString& ExistingEntrySourceLocation = (!ExistingEntryFileName.IsEmpty()) ? ExistingEntryFileName : ConflictingContext->SourceLocation;
 
 			FString Message = SanitizeLogOutput(
 				FString::Printf(TEXT("Found previously entered localized string: %s [%s] %s %s=\"%s\" %s. It was previously \"%s\" %s in %s."),
@@ -1337,44 +1419,70 @@ bool FLocTextHelper::FindKeysForLegacyTranslation(const TSharedRef<const FIntern
 
 TSharedPtr<FInternationalizationManifest> FLocTextHelper::LoadManifestImpl(const FString& InManifestFilePath, const ELocTextHelperLoadFlags InLoadFlags, FText* OutError)
 {
-	TSharedPtr<FInternationalizationManifest> LocalManifest = MakeShareable(new FInternationalizationManifest());
+	TSharedRef<FInternationalizationManifest> LocalManifest = MakeShared<FInternationalizationManifest>();
 
-	// Attempt to load an existing manifest first
-	if (!!(InLoadFlags & ELocTextHelperLoadFlags::Load))
+	auto LoadSingleManifest = [this, &LocalManifest, &OutError](const FString& InManifestFilePathToLoad, const FName InPlatformName) -> bool
 	{
 		bool bLoaded = false;
-		const bool bExisted = FPaths::FileExists(InManifestFilePath);
 
-		if (bExisted)
+		if (LocFileNotifies.IsValid())
 		{
-			if (LocFileNotifies.IsValid())
-			{
-				LocFileNotifies->PreFileRead(InManifestFilePath);
-			}
+			LocFileNotifies->PreFileRead(InManifestFilePathToLoad);
+		}
 
-			if (FJsonInternationalizationManifestSerializer::DeserializeManifestFromFile(InManifestFilePath, LocalManifest.ToSharedRef()))
-			{
-				bLoaded = true;
-			}
-			else
-			{
-				if (OutError)
-				{
-					*OutError = FText::Format(LOCTEXT("Error_LoadManifest_DeserializeFile", "Failed to deserialize manifest '{0}'."), FText::FromString(InManifestFilePath));
-				}
-			}
+		if (FJsonInternationalizationManifestSerializer::DeserializeManifestFromFile(InManifestFilePathToLoad, LocalManifest, InPlatformName))
+		{
+			bLoaded = true;
+		}
+		else if (OutError)
+		{
+			*OutError = FText::Format(LOCTEXT("Error_LoadManifest_DeserializeFile", "Failed to deserialize manifest '{0}'."), FText::FromString(InManifestFilePathToLoad));
+		}
+
+		if (LocFileNotifies.IsValid())
+		{
+			LocFileNotifies->PostFileRead(InManifestFilePathToLoad);
+		}
+
+		return bLoaded;
+	};
 
-			if (LocFileNotifies.IsValid())
+	// Attempt to load an existing manifest first
+	if (!!(InLoadFlags & ELocTextHelperLoadFlags::Load))
+	{
+		const bool bExists = FPaths::FileExists(InManifestFilePath);
+
+		bool bLoadedAll = bExists;
+		if (bExists)
+		{
+			bLoadedAll &= LoadSingleManifest(InManifestFilePath, FName());
 			{
-				LocFileNotifies->PostFileRead(InManifestFilePath);
+				// Load all per-platform manifests too
+				// We always do this, as we may have changed the split config so don't want to lose data
+				const FString PlatformManifestName = FPaths::GetCleanFilename(InManifestFilePath);
+				const FString PlatformLocalizationPath = FPaths::GetPath(InManifestFilePath) / FPaths::GetPlatformLocalizationFolderName();
+				IFileManager::Get().IterateDirectory(*PlatformLocalizationPath, [&](const TCHAR* FilenameOrDirectory, bool bIsDirectory) -> bool
+				{
+					if (bIsDirectory)
+					{
+						const FString PlatformManifestFilePath = FilenameOrDirectory / PlatformManifestName;
+						if (FPaths::FileExists(PlatformManifestFilePath))
+						{
+							const FString SplitPlatformName = FPaths::GetCleanFilename(FilenameOrDirectory);
+							bLoadedAll &= LoadSingleManifest(PlatformManifestFilePath, *SplitPlatformName);
+						}
+					}
+					return true;
+				});
 			}
 		}
 
-		if (bLoaded)
+		if (bLoadedAll)
 		{
 			return LocalManifest;
 		}
-		else if (bExisted)
+
+		if (bExists)
 		{
 			// Don't allow fallback to Create if the file exists but could not be loaded
 			return nullptr;
@@ -1392,80 +1500,170 @@ TSharedPtr<FInternationalizationManifest> FLocTextHelper::LoadManifestImpl(const
 
 bool FLocTextHelper::SaveManifestImpl(const TSharedRef<const FInternationalizationManifest>& InManifest, const FString& InManifestFilePath, FText* OutError) const
 {
-	bool bSaved = false;
-
-	if (LocFileNotifies.IsValid())
+	auto SaveSingleManifest = [this, &OutError](const TSharedRef<const FInternationalizationManifest>& InManifestToSave, const FString& InManifestFilePathToSave) -> bool
 	{
-		LocFileNotifies->PreFileWrite(InManifestFilePath);
-	}
+		bool bSaved = false;
 
-	TSharedRef<FJsonObject> ManifestJsonObject = MakeShareable(new FJsonObject());
-	if (FJsonInternationalizationManifestSerializer::SerializeManifestToFile(InManifest, InManifestFilePath))
-	{
-		bSaved = true;
-	}
-	else
+		if (LocFileNotifies.IsValid())
+		{
+			LocFileNotifies->PreFileWrite(InManifestFilePathToSave);
+		}
+
+		if (FJsonInternationalizationManifestSerializer::SerializeManifestToFile(InManifestToSave, InManifestFilePathToSave))
+		{
+			bSaved = true;
+		}
+		else
+		{
+			if (OutError)
+			{
+				*OutError = FText::Format(LOCTEXT("Error_SaveManifest_SerializeFile", "Failed to serialize manifest '{0}'."), FText::FromString(InManifestFilePathToSave));
+			}
+		}
+
+		if (LocFileNotifies.IsValid())
+		{
+			LocFileNotifies->PostFileWrite(InManifestFilePathToSave);
+		}
+
+		return bSaved;
+	};
+
+	bool bSavedAll = true;
+	if (ShouldSplitPlatformData())
 	{
-		if (OutError)
+		const FString PlatformManifestName = FPaths::GetCleanFilename(InManifestFilePath);
+		const FString PlatformLocalizationPath = FPaths::GetPath(InManifestFilePath) / FPaths::GetPlatformLocalizationFolderName();
+
+		// Split the manifest into separate entries for each platform, as well as a platform agnostic manifest
+		TSharedRef<FInternationalizationManifest> PlatformAgnosticManifest = MakeShared<FInternationalizationManifest>();
+		TMap<FName, TSharedRef<FInternationalizationManifest>> PerPlatformManifests;
 		{
-			*OutError = FText::Format(LOCTEXT("Error_SaveManifest_SerializeFile", "Failed to serialize manifest '{0}'."), FText::FromString(InManifestFilePath));
+			// Always add the split platforms so that they generate an empty manifest if there are no entries for that platform in the master manifest
+			for (const FString& SplitPlatformName : GetPlatformsToSplit())
+			{
+				PerPlatformManifests.Add(*SplitPlatformName, MakeShared<FInternationalizationManifest>());
+			}
+
+			// Split the manifest entries based on the platform they belonged to
+			for (FManifestEntryByStringContainer::TConstIterator It(InManifest->GetEntriesBySourceTextIterator()); It; ++It)
+			{
+				const TSharedRef<FManifestEntry> ManifestEntry = It.Value();
+				for (const FManifestContext& Context : ManifestEntry->Contexts)
+				{
+					TSharedPtr<FInternationalizationManifest> ManifestToUpdate = PlatformAgnosticManifest;
+					if (!Context.PlatformName.IsNone())
+					{
+						if (TSharedRef<FInternationalizationManifest>* PerPlatformManifest = PerPlatformManifests.Find(Context.PlatformName))
+						{
+							ManifestToUpdate = *PerPlatformManifest;
+						}
+					}
+					check(ManifestToUpdate.IsValid());
+
+					if (!ManifestToUpdate->AddSource(ManifestEntry->Namespace, ManifestEntry->Source, Context))
+					{
+						UE_LOG(LogLocTextHelper, Error, TEXT("Could not process localized string: %s [%s] %s=\"%s\" %s."),
+							*ManifestEntry->Namespace.GetString(),
+							*Context.Key.GetString(),
+							*ManifestEntry->Source.Text,
+							*FJsonInternationalizationMetaDataSerializer::MetadataToString(ManifestEntry->Source.MetadataObj)
+							);
+					}
+				}
+			}
 		}
-	}
 
-	if (LocFileNotifies.IsValid())
+		bSavedAll &= SaveSingleManifest(PlatformAgnosticManifest, InManifestFilePath);
+		for (const auto PerPlatformManifestPair : PerPlatformManifests)
+		{
+			const FString PlatformManifestFilePath = PlatformLocalizationPath / PerPlatformManifestPair.Key.ToString() / PlatformManifestName;
+			bSavedAll &= SaveSingleManifest(PerPlatformManifestPair.Value, PlatformManifestFilePath);
+		}
+	}
+	else
 	{
-		LocFileNotifies->PostFileWrite(InManifestFilePath);
+		bSavedAll &= SaveSingleManifest(InManifest, InManifestFilePath);
 	}
-
-	return bSaved;
+	return bSavedAll;
 }
 
 TSharedPtr<FInternationalizationArchive> FLocTextHelper::LoadArchiveImpl(const FString& InArchiveFilePath, const ELocTextHelperLoadFlags InLoadFlags, FText* OutError)
 {
-	TSharedPtr<FInternationalizationArchive> LocalArchive = MakeShareable(new FInternationalizationArchive());
+	TSharedRef<FInternationalizationArchive> LocalArchive = MakeShared<FInternationalizationArchive>();
 
-	// Attempt to load an existing archive first
-	if (!!(InLoadFlags & ELocTextHelperLoadFlags::Load))
+	auto LoadSingleArchive = [this, &LocalArchive, &OutError](const FString& InArchiveFilePathToLoad) -> bool
 	{
 		bool bLoaded = false;
-		const bool bExisted = FPaths::FileExists(InArchiveFilePath);
 
-		if (bExisted)
+		if (LocFileNotifies.IsValid())
 		{
-			if (LocFileNotifies.IsValid())
-			{
-				LocFileNotifies->PreFileRead(InArchiveFilePath);
-			}
+			LocFileNotifies->PreFileRead(InArchiveFilePathToLoad);
+		}
 
-			TSharedPtr<FInternationalizationArchive> NativeArchive;
-			if (!NativeCulture.IsEmpty())
-			{
-				NativeArchive = Archives.FindRef(NativeCulture);
-			}
-			
-			if (FJsonInternationalizationArchiveSerializer::DeserializeArchiveFromFile(InArchiveFilePath, LocalArchive.ToSharedRef(), Manifest, NativeArchive))
-			{
-				bLoaded = true;
-			}
-			else
+		TSharedPtr<FInternationalizationArchive> NativeArchive;
+		if (!NativeCulture.IsEmpty())
+		{
+			NativeArchive = Archives.FindRef(NativeCulture);
+		}
+
+		if (FJsonInternationalizationArchiveSerializer::DeserializeArchiveFromFile(InArchiveFilePathToLoad, LocalArchive, Manifest, NativeArchive))
+		{
+			bLoaded = true;
+		}
+		else
+		{
+			if (OutError)
 			{
-				if (OutError)
-				{
-					*OutError = FText::Format(LOCTEXT("Error_LoadArchive_DeserializeFile", "Failed to deserialize archive '{0}'."), FText::FromString(InArchiveFilePath));
-				}
+				*OutError = FText::Format(LOCTEXT("Error_LoadArchive_DeserializeFile", "Failed to deserialize archive '{0}'."), FText::FromString(InArchiveFilePathToLoad));
 			}
+		}
+
+		if (LocFileNotifies.IsValid())
+		{
+			LocFileNotifies->PostFileRead(InArchiveFilePathToLoad);
+		}
 
-			if (LocFileNotifies.IsValid())
+		return bLoaded;
+	};
+
+	// Attempt to load an existing archive first
+	if (!!(InLoadFlags & ELocTextHelperLoadFlags::Load))
+	{
+		const bool bExists = FPaths::FileExists(InArchiveFilePath);
+
+		bool bLoadedAll = bExists;
+		if (bExists)
+		{
+			bLoadedAll &= LoadSingleArchive(InArchiveFilePath);
 			{
-				LocFileNotifies->PostFileRead(InArchiveFilePath);
+				// Load all per-platform archives too
+				// We always do this, as we may have changed the split config so don't want to lose data
+				const FString ArchiveCultureFilePath = FPaths::GetPath(InArchiveFilePath);
+				const FString PlatformArchiveName = FPaths::GetCleanFilename(InArchiveFilePath);
+				const FString PlatformArchiveCulture = FPaths::GetCleanFilename(ArchiveCultureFilePath);
+				const FString PlatformLocalizationPath = FPaths::GetPath(ArchiveCultureFilePath) / FPaths::GetPlatformLocalizationFolderName();
+				IFileManager::Get().IterateDirectory(*PlatformLocalizationPath, [&](const TCHAR* FilenameOrDirectory, bool bIsDirectory) -> bool
+				{
+					if (bIsDirectory)
+					{
+						const FString PlatformArchiveFilePath = FilenameOrDirectory / PlatformArchiveCulture / PlatformArchiveName;
+						if (FPaths::FileExists(PlatformArchiveFilePath))
+						{
+							bLoadedAll &= LoadSingleArchive(PlatformArchiveFilePath);
+						}
+					}
+					return true;
+				});
 			}
 		}
 
-		if (bLoaded)
+		if (bLoadedAll)
 		{
 			return LocalArchive;
 		}
-		else if (bExisted)
+		
+		if (bExists)
 		{
 			// Don't allow fallback to Create if the file exists but could not be loaded
 			return nullptr;
@@ -1483,31 +1681,92 @@ TSharedPtr<FInternationalizationArchive> FLocTextHelper::LoadArchiveImpl(const F
 
 bool FLocTextHelper::SaveArchiveImpl(const TSharedRef<const FInternationalizationArchive>& InArchive, const FString& InArchiveFilePath, FText* OutError) const
 {
-	bool bSaved = false;
-
-	if (LocFileNotifies.IsValid())
+	auto SaveSingleArchive = [this, &OutError](const TSharedRef<const FInternationalizationArchive>& InArchiveToSave, const FString& InArchiveFilePathToSave) -> bool
 	{
-		LocFileNotifies->PreFileWrite(InArchiveFilePath);
-	}
+		bool bSaved = false;
 
-	if (FJsonInternationalizationArchiveSerializer::SerializeArchiveToFile(InArchive, InArchiveFilePath))
-	{
-		bSaved = true;
-	}
-	else
+		if (LocFileNotifies.IsValid())
+		{
+			LocFileNotifies->PreFileWrite(InArchiveFilePathToSave);
+		}
+
+		if (FJsonInternationalizationArchiveSerializer::SerializeArchiveToFile(InArchiveToSave, InArchiveFilePathToSave))
+		{
+			bSaved = true;
+		}
+		else
+		{
+			if (OutError)
+			{
+				*OutError = FText::Format(LOCTEXT("Error_SaveArchive_SerializeFile", "Failed to serialize archive '{0}'."), FText::FromString(InArchiveFilePathToSave));
+			}
+		}
+
+		if (LocFileNotifies.IsValid())
+		{
+			LocFileNotifies->PostFileWrite(InArchiveFilePathToSave);
+		}
+
+		return bSaved;
+	};
+
+	bool bSavedAll = true;
+	if (ShouldSplitPlatformData())
 	{
-		if (OutError)
+		const FString ArchiveCultureFilePath = FPaths::GetPath(InArchiveFilePath);
+		const FString PlatformArchiveName = FPaths::GetCleanFilename(InArchiveFilePath);
+		const FString PlatformArchiveCulture = FPaths::GetCleanFilename(ArchiveCultureFilePath);
+		const FString PlatformLocalizationPath = FPaths::GetPath(ArchiveCultureFilePath) / FPaths::GetPlatformLocalizationFolderName();
+
+		// Split the archive into separate entries for each platform, as well as a platform agnostic archive
+		TSharedRef<FInternationalizationArchive> PlatformAgnosticArchive = MakeShared<FInternationalizationArchive>();
+		TMap<FName, TSharedRef<FInternationalizationArchive>> PerPlatformArchives;
+		{
+			// Always add the split platforms so that they generate an empty archives if there are no entries for that platform in the master archive
+			for (const FString& SplitPlatformName : GetPlatformsToSplit())
+			{
+				PerPlatformArchives.Add(*SplitPlatformName, MakeShared<FInternationalizationArchive>());
+			}
+
+			// Split the archive entries based on the platform they belonged to
+			EnumerateSourceTexts([&InArchive, &PlatformAgnosticArchive, &PerPlatformArchives](TSharedRef<FManifestEntry> InManifestEntry) -> bool
+			{
+				for (const FManifestContext& Context : InManifestEntry->Contexts)
+				{
+					TSharedPtr<FInternationalizationArchive> ArchiveToUpdate = PlatformAgnosticArchive;
+					if (!Context.PlatformName.IsNone())
+					{
+						if (TSharedRef<FInternationalizationArchive>* PerPlatformArchive = PerPlatformArchives.Find(Context.PlatformName))
+						{
+							ArchiveToUpdate = *PerPlatformArchive;
+						}
+					}
+					check(ArchiveToUpdate.IsValid());
+
+					// Keep any translation for the source text
+					TSharedPtr<FArchiveEntry> ArchiveEntry = InArchive->FindEntryByKey(InManifestEntry->Namespace, Context.Key, Context.KeyMetadataObj);
+					if (ArchiveEntry.IsValid())
+					{
+						ArchiveToUpdate->AddEntry(ArchiveEntry.ToSharedRef());
+					}
+				}
+
+				return true; // continue enumeration
+			}, true);
+		}
+
+		bSavedAll &= SaveSingleArchive(PlatformAgnosticArchive, InArchiveFilePath);
+		for (const auto PerPlatformArchivePair : PerPlatformArchives)
 		{
-			*OutError = FText::Format(LOCTEXT("Error_SaveArchive_SerializeFile", "Failed to serialize archive '{0}'."), FText::FromString(InArchiveFilePath));
+			const FString PlatformArchiveFilePath = PlatformLocalizationPath / PerPlatformArchivePair.Key.ToString() / PlatformArchiveCulture / PlatformArchiveName;
+			bSavedAll &= SaveSingleArchive(PerPlatformArchivePair.Value, PlatformArchiveFilePath);
 		}
 	}
-
-	if (LocFileNotifies.IsValid())
+	else
 	{
-		LocFileNotifies->PostFileWrite(InArchiveFilePath);
+		bSavedAll &= SaveSingleArchive(InArchive, InArchiveFilePath);
 	}
-
-	return bSaved;
+	return bSavedAll;
 }
 
 TSharedPtr<FArchiveEntry> FLocTextHelper::FindTranslationImpl(const FString& InCulture, const FLocKey& InNamespace, const FLocKey& InKey, const TSharedPtr<FLocMetadataObject> InKeyMetadataObj) const
diff --git a/Engine/Source/Developer/Localization/Private/PortableObjectPipeline.cpp b/Engine/Source/Developer/Localization/Private/PortableObjectPipeline.cpp
--- a/Engine/Source/Developer/Localization/Private/PortableObjectPipeline.cpp
+++ b/Engine/Source/Developer/Localization/Private/PortableObjectPipeline.cpp
@@ -51,9 +51,14 @@ namespace
 		FLocKeyPairMultiMap CollapsedNSSourceStringToExpandedNSKey;
 	};
 
-	TSharedRef<FInternationalizationManifest> BuildCollapsedManifest(FLocTextHelper& InLocTextHelper, const ELocalizedTextCollapseMode InTextCollapseMode, FCollapsedData& OutCollapsedData)
+	void BuildCollapsedManifest(FLocTextHelper& InLocTextHelper, const ELocalizedTextCollapseMode InTextCollapseMode, FCollapsedData& OutCollapsedData, TSharedPtr<FInternationalizationManifest>& OutPlatformAgnosticManifest, TMap<FName, TSharedRef<FInternationalizationManifest>>& OutPerPlatformManifests)
 	{
-		TSharedRef<FInternationalizationManifest> CollapsedManifest = MakeShared<FInternationalizationManifest>();
+		// Always add the split platforms so that they generate an empty manifest if there are no entries for that platform in the master manifest
+		OutPlatformAgnosticManifest = MakeShared<FInternationalizationManifest>();
+		for (const FString& SplitPlatformName : InLocTextHelper.GetPlatformsToSplit())
+		{
+			OutPerPlatformManifests.Add(*SplitPlatformName, MakeShared<FInternationalizationManifest>());
+		}
 
 		InLocTextHelper.EnumerateSourceTexts([&](TSharedRef<FManifestEntry> InManifestEntry) -> bool
 		{
@@ -63,8 +68,18 @@ namespace
 			{
 				bool bAddedContext = false;
 
+				TSharedPtr<FInternationalizationManifest> ManifestToUpdate = OutPlatformAgnosticManifest;
+				if (!Context.PlatformName.IsNone())
+				{
+					if (TSharedRef<FInternationalizationManifest>* PerPlatformManifest = OutPerPlatformManifests.Find(Context.PlatformName))
+					{
+						ManifestToUpdate = *PerPlatformManifest;
+					}
+				}
+				check(ManifestToUpdate.IsValid());
+
 				// Check if the entry already exists in the manifest
-				TSharedPtr<FManifestEntry> ExistingEntry = CollapsedManifest->FindEntryByContext(CollapsedNamespace, Context);
+				TSharedPtr<FManifestEntry> ExistingEntry = ManifestToUpdate->FindEntryByContext(CollapsedNamespace, Context);
 				if (ExistingEntry.IsValid())
 				{
 					if (InManifestEntry->Source.IsExactMatch(ExistingEntry->Source))
@@ -97,7 +112,7 @@ namespace
 				}
 				else
 				{
-					if (CollapsedManifest->AddSource(CollapsedNamespace, InManifestEntry->Source, Context))
+					if (ManifestToUpdate->AddSource(CollapsedNamespace, InManifestEntry->Source, Context))
 					{
 						bAddedContext = true;
 					}
@@ -137,8 +152,6 @@ namespace
 
 			return true; // continue enumeration
 		}, true);
-
-		return CollapsedManifest;
 	}
 
 	TMap<FPortableObjectEntryKey, TArray<FString>> ExtractPreservedPOComments(const FPortableObjectFormatDOM& InPortableObject)
@@ -416,9 +429,18 @@ namespace
 
 bool PortableObjectPipeline::Import(FLocTextHelper& InLocTextHelper, const FString& InCulture, const FString& InPOFilePath, const ELocalizedTextCollapseMode InTextCollapseMode)
 {
+	// This function only works when not splitting per-platform data
+	if (InLocTextHelper.ShouldSplitPlatformData())
+	{
+		UE_LOG(LogPortableObjectPipeline, Error, TEXT("PortableObjectPipeline::Import may only be used when not splitting platform data."));
+		return false;
+	}
+
 	// Build the collapsed manifest data needed to import
 	FCollapsedData CollapsedData;
-	TSharedRef<FInternationalizationManifest> CollapsedManifest = BuildCollapsedManifest(InLocTextHelper, InTextCollapseMode, CollapsedData);
+	TSharedPtr<FInternationalizationManifest> PlatformAgnosticManifest;
+	TMap<FName, TSharedRef<FInternationalizationManifest>> PerPlatformManifests;
+	BuildCollapsedManifest(InLocTextHelper, InTextCollapseMode, CollapsedData, PlatformAgnosticManifest, PerPlatformManifests);
 
 	return ImportPortableObject(InLocTextHelper, InCulture, InPOFilePath, CollapsedData);
 }
@@ -427,32 +449,51 @@ bool PortableObjectPipeline::ImportAll(FLocTextHelper& InLocTextHelper, const FS
 {
 	// We may only have a single culture if using this setting
 	const bool bSingleCultureMode = !bUseCultureDirectory;
-	if (bSingleCultureMode && InLocTextHelper.GetAllCultures(bSingleCultureMode).Num() != 1)
+	if (bSingleCultureMode && (InLocTextHelper.GetAllCultures(bSingleCultureMode).Num() != 1 || InLocTextHelper.ShouldSplitPlatformData()))
 	{
-		UE_LOG(LogPortableObjectPipeline, Error, TEXT("bUseCultureDirectory may only be used with a single culture."));
+		UE_LOG(LogPortableObjectPipeline, Error, TEXT("bUseCultureDirectory may only be used with a single culture when not splitting platform data."));
 		return false;
 	}
 
 	// Build the collapsed manifest data needed to import
 	FCollapsedData CollapsedData;
-	TSharedRef<FInternationalizationManifest> CollapsedManifest = BuildCollapsedManifest(InLocTextHelper, InTextCollapseMode, CollapsedData);
+	TSharedPtr<FInternationalizationManifest> PlatformAgnosticManifest;
+	TMap<FName, TSharedRef<FInternationalizationManifest>> PerPlatformManifests;
+	BuildCollapsedManifest(InLocTextHelper, InTextCollapseMode, CollapsedData, PlatformAgnosticManifest, PerPlatformManifests);
 
 	// Process the desired cultures
 	bool bSuccess = true;
 	for (const FString& CultureName : InLocTextHelper.GetAllCultures(bSingleCultureMode))
 	{
-		// Which path should we use for the PO?
-		FString POFilePath;
-		if (bUseCultureDirectory)
+		auto ImportSinglePortableObject = [&](const FName InPlatformName) -> bool
 		{
-			POFilePath = InPOCultureRootPath / CultureName / InPOFilename;
-		}
-		else
+			// Which path should we use for the PO?
+			FString POFilePath;
+			if (bUseCultureDirectory)
+			{
+				// Platforms splits are only supported when exporting all cultures (see the error at the start of this function)
+				if (InPlatformName.IsNone())
+				{
+					POFilePath = InPOCultureRootPath / CultureName / InPOFilename;
+				}
+				else
+				{
+					POFilePath = InPOCultureRootPath / FPaths::GetPlatformLocalizationFolderName() / InPlatformName.ToString() / CultureName / InPOFilename;
+				}
+			}
+			else
+			{
+				POFilePath = InPOCultureRootPath / InPOFilename;
+			}
+
+			return ImportPortableObject(InLocTextHelper, CultureName, POFilePath, CollapsedData);
+		};
+
+		bSuccess &= ImportSinglePortableObject(FName());
+		for (const auto& PerPlatformManifestPair : PerPlatformManifests)
 		{
-			POFilePath = InPOCultureRootPath / InPOFilename;
+			bSuccess &= ImportSinglePortableObject(PerPlatformManifestPair.Key);
 		}
-
-		bSuccess &= ImportPortableObject(InLocTextHelper, CultureName, POFilePath, CollapsedData);
 	}
 
 	return bSuccess;
@@ -460,20 +501,29 @@ bool PortableObjectPipeline::ImportAll(FLocTextHelper& InLocTextHelper, const FS
 
 bool PortableObjectPipeline::Export(FLocTextHelper& InLocTextHelper, const FString& InCulture, const FString& InPOFilePath, const ELocalizedTextCollapseMode InTextCollapseMode, const bool bShouldPersistComments)
 {
+	// This function only works when not splitting per-platform data
+	if (InLocTextHelper.ShouldSplitPlatformData())
+	{
+		UE_LOG(LogPortableObjectPipeline, Error, TEXT("PortableObjectPipeline::Export may only be used when not splitting platform data."));
+		return false;
+	}
+
 	// Build the collapsed manifest data needed to export
 	FCollapsedData CollapsedData;
-	TSharedRef<FInternationalizationManifest> CollapsedManifest = BuildCollapsedManifest(InLocTextHelper, InTextCollapseMode, CollapsedData);
+	TSharedPtr<FInternationalizationManifest> PlatformAgnosticManifest;
+	TMap<FName, TSharedRef<FInternationalizationManifest>> PerPlatformManifests;
+	BuildCollapsedManifest(InLocTextHelper, InTextCollapseMode, CollapsedData, PlatformAgnosticManifest, PerPlatformManifests);
 
-	return ExportPortableObject(InLocTextHelper, InCulture, InPOFilePath, InTextCollapseMode, CollapsedManifest, CollapsedData, bShouldPersistComments);
+	return ExportPortableObject(InLocTextHelper, InCulture, InPOFilePath, InTextCollapseMode, PlatformAgnosticManifest.ToSharedRef(), CollapsedData, bShouldPersistComments);
 }
 
 bool PortableObjectPipeline::ExportAll(FLocTextHelper& InLocTextHelper, const FString& InPOCultureRootPath, const FString& InPOFilename, const ELocalizedTextCollapseMode InTextCollapseMode, const bool bShouldPersistComments, const bool bUseCultureDirectory)
 {
 	// We may only have a single culture if using this setting
 	const bool bSingleCultureMode = !bUseCultureDirectory;
-	if (bSingleCultureMode && InLocTextHelper.GetAllCultures(bSingleCultureMode).Num() != 1)
+	if (bSingleCultureMode && (InLocTextHelper.GetAllCultures(bSingleCultureMode).Num() != 1 || InLocTextHelper.ShouldSplitPlatformData()))
 	{
-		UE_LOG(LogPortableObjectPipeline, Error, TEXT("bUseCultureDirectory may only be used with a single culture."));
+		UE_LOG(LogPortableObjectPipeline, Error, TEXT("bUseCultureDirectory may only be used with a single culture when not splitting platform data."));
 		return false;
 	}
 
@@ -486,24 +536,43 @@ bool PortableObjectPipeline::ExportAll(FLocTextHelper& InLocTextHelper, const FS
 
 	// Build the collapsed manifest data to export
 	FCollapsedData CollapsedData;
-	TSharedRef<FInternationalizationManifest> CollapsedManifest = BuildCollapsedManifest(InLocTextHelper, InTextCollapseMode, CollapsedData);
+	TSharedPtr<FInternationalizationManifest> PlatformAgnosticManifest;
+	TMap<FName, TSharedRef<FInternationalizationManifest>> PerPlatformManifests;
+	BuildCollapsedManifest(InLocTextHelper, InTextCollapseMode, CollapsedData, PlatformAgnosticManifest, PerPlatformManifests);
 
 	// Process the desired cultures
 	bool bSuccess = true;
 	for (const FString& CultureName : InLocTextHelper.GetAllCultures(bSingleCultureMode))
 	{
-		// Which path should we use for the PO?
-		FString POFilePath;
-		if (bUseCultureDirectory)
+		auto ExportSinglePortableObject = [&](const TSharedRef<FInternationalizationManifest>& InCollapsedManifest, const FName InPlatformName) -> bool
 		{
-			POFilePath = InPOCultureRootPath / CultureName / InPOFilename;
-		}
-		else
+			// Which path should we use for the PO?
+			FString POFilePath;
+			if (bUseCultureDirectory)
+			{
+				// Platforms splits are only supported when exporting all cultures (see the error at the start of this function)
+				if (InPlatformName.IsNone())
+				{
+					POFilePath = InPOCultureRootPath / CultureName / InPOFilename;
+				}
+				else
+				{
+					POFilePath = InPOCultureRootPath / FPaths::GetPlatformLocalizationFolderName() / InPlatformName.ToString() / CultureName / InPOFilename;
+				}
+			}
+			else
+			{
+				POFilePath = InPOCultureRootPath / InPOFilename;
+			}
+
+			return ExportPortableObject(InLocTextHelper, CultureName, POFilePath, InTextCollapseMode, InCollapsedManifest, CollapsedData, bShouldPersistComments);
+		};
+
+		bSuccess &= ExportSinglePortableObject(PlatformAgnosticManifest.ToSharedRef(), FName());
+		for (const auto& PerPlatformManifestPair : PerPlatformManifests)
 		{
-			POFilePath = InPOCultureRootPath / InPOFilename;
+			bSuccess &= ExportSinglePortableObject(PerPlatformManifestPair.Value, PerPlatformManifestPair.Key);
 		}
-
-		bSuccess &= ExportPortableObject(InLocTextHelper, CultureName, POFilePath, InTextCollapseMode, CollapsedManifest, CollapsedData, bShouldPersistComments);
 	}
 
 	return bSuccess;
diff --git a/Engine/Source/Developer/Localization/Private/Serialization/JsonInternationalizationManifestSerializer.cpp b/Engine/Source/Developer/Localization/Private/Serialization/JsonInternationalizationManifestSerializer.cpp
--- a/Engine/Source/Developer/Localization/Private/Serialization/JsonInternationalizationManifestSerializer.cpp
+++ b/Engine/Source/Developer/Localization/Private/Serialization/JsonInternationalizationManifestSerializer.cpp
@@ -46,7 +46,7 @@ struct FCompareStructuredEntryByNamespace
 };
 
 
-bool FJsonInternationalizationManifestSerializer::DeserializeManifest( const FString& InStr, TSharedRef< FInternationalizationManifest > Manifest )
+bool FJsonInternationalizationManifestSerializer::DeserializeManifest( const FString& InStr, TSharedRef< FInternationalizationManifest > Manifest, const FName PlatformName )
 {
 	TSharedPtr<FJsonObject> JsonManifestObj;
 	TSharedRef< TJsonReader<> > Reader = TJsonReaderFactory<>::Create( InStr );
@@ -54,20 +54,20 @@ bool FJsonInternationalizationManifestSerializer::DeserializeManifest( const FSt
 
 	if ( bExecSuccessful && JsonManifestObj.IsValid() )
 	{
-		bExecSuccessful = DeserializeInternal( JsonManifestObj.ToSharedRef(), Manifest );
+		bExecSuccessful = DeserializeInternal( JsonManifestObj.ToSharedRef(), Manifest, PlatformName );
 	}
 
 	return bExecSuccessful;
 }
 
 
-bool FJsonInternationalizationManifestSerializer::DeserializeManifest( TSharedRef< FJsonObject > InJsonObj, TSharedRef< FInternationalizationManifest > Manifest )
+bool FJsonInternationalizationManifestSerializer::DeserializeManifest( TSharedRef< FJsonObject > InJsonObj, TSharedRef< FInternationalizationManifest > Manifest, const FName PlatformName )
 {
-	return DeserializeInternal( InJsonObj, Manifest );
+	return DeserializeInternal( InJsonObj, Manifest, PlatformName );
 }
 
 
-bool FJsonInternationalizationManifestSerializer::DeserializeManifestFromFile( const FString& InJsonFile, TSharedRef< FInternationalizationManifest > Manifest )
+bool FJsonInternationalizationManifestSerializer::DeserializeManifestFromFile( const FString& InJsonFile, TSharedRef< FInternationalizationManifest > Manifest, const FName PlatformName )
 {
 	// Read in file as string
 	FString FileContents;
@@ -87,7 +87,7 @@ bool FJsonInternationalizationManifestSerializer::DeserializeManifestFromFile( c
 		return false;
 	}
 
-	return DeserializeInternal(JsonObject.ToSharedRef(), Manifest);
+	return DeserializeInternal(JsonObject.ToSharedRef(), Manifest, PlatformName);
 }
 
 
@@ -137,7 +137,7 @@ bool FJsonInternationalizationManifestSerializer::SerializeManifestToFile( TShar
 }
 
 
-bool FJsonInternationalizationManifestSerializer::DeserializeInternal( TSharedRef< FJsonObject > InJsonObj, TSharedRef< FInternationalizationManifest > Manifest )
+bool FJsonInternationalizationManifestSerializer::DeserializeInternal( TSharedRef< FJsonObject > InJsonObj, TSharedRef< FInternationalizationManifest > Manifest, const FName PlatformName )
 {
 	if( InJsonObj->HasField( TAG_FORMATVERSION ) )
 	{
@@ -155,7 +155,7 @@ bool FJsonInternationalizationManifestSerializer::DeserializeInternal( TSharedRe
 		Manifest->SetFormatVersion(FInternationalizationManifest::EFormatVersion::Initial);
 	}
 
-	return JsonObjToManifest(InJsonObj, TEXT(""), Manifest );
+	return JsonObjToManifest(InJsonObj, TEXT(""), Manifest, PlatformName );
 }
 
 
@@ -182,7 +182,7 @@ bool FJsonInternationalizationManifestSerializer::SerializeInternal( TSharedRef<
 }
 
 
-bool FJsonInternationalizationManifestSerializer::JsonObjToManifest( TSharedRef< FJsonObject > InJsonObj, FString ParentNamespace, TSharedRef< FInternationalizationManifest > Manifest )
+bool FJsonInternationalizationManifestSerializer::JsonObjToManifest( TSharedRef< FJsonObject > InJsonObj, FString ParentNamespace, TSharedRef< FInternationalizationManifest > Manifest, const FName PlatformName )
 {
 	bool bConvertSuccess = true;
 	FString AccumulatedNamespace = ParentNamespace;
@@ -266,6 +266,7 @@ bool FJsonInternationalizationManifestSerializer::JsonObjToManifest( TSharedRef<
 						FManifestContext CommandContext;
 						CommandContext.Key = Key;
 						CommandContext.SourceLocation = SourceLocation;
+						CommandContext.PlatformName = PlatformName;
 
 						if( ContextJSONObject->HasTypedField< EJson::Boolean >( TAG_OPTIONAL ) )
 						{
@@ -337,7 +338,7 @@ bool FJsonInternationalizationManifestSerializer::JsonObjToManifest( TSharedRef<
 			const TSharedPtr< FJsonValue >  SubnamespaceEntry = *SubnamespaceIter;
 			const TSharedPtr< FJsonObject > SubnamespaceJSONObject = SubnamespaceEntry->AsObject();
 
-			if( !JsonObjToManifest( SubnamespaceJSONObject.ToSharedRef(), AccumulatedNamespace, Manifest ) )
+			if( !JsonObjToManifest( SubnamespaceJSONObject.ToSharedRef(), AccumulatedNamespace, Manifest, PlatformName ) )
 			{
 				bConvertSuccess = false;
 				break;
diff --git a/Engine/Source/Developer/Localization/Private/TextLocalizationResourceGenerator.cpp b/Engine/Source/Developer/Localization/Private/TextLocalizationResourceGenerator.cpp
--- a/Engine/Source/Developer/Localization/Private/TextLocalizationResourceGenerator.cpp
+++ b/Engine/Source/Developer/Localization/Private/TextLocalizationResourceGenerator.cpp
@@ -18,13 +18,23 @@ bool FTextLocalizationResourceGenerator::GenerateLocMeta(const FLocTextHelper& I
 	return true;
 }
 
-bool FTextLocalizationResourceGenerator::GenerateLocRes(const FLocTextHelper& InLocTextHelper, const FString& InCultureToGenerate, const bool bSkipSourceCheck, const FTextKey& InLocResID, FTextLocalizationResource& OutLocRes, const int32 InPriority)
+bool FTextLocalizationResourceGenerator::GenerateLocRes(const FLocTextHelper& InLocTextHelper, const FString& InCultureToGenerate, const bool bSkipSourceCheck, const FTextKey& InLocResID, FTextLocalizationResource& OutPlatformAgnosticLocRes, TMap<FName, TSharedRef<FTextLocalizationResource>>& OutPerPlatformLocRes, const int32 InPriority)
 {
 	const bool bIsNativeCulture = InCultureToGenerate == InLocTextHelper.GetNativeCulture();
 	FCulturePtr Culture = FInternationalization::Get().GetCulture(InCultureToGenerate);
 
+	// Always add the split platforms so that they generate an empty LocRes if there are no entries for that platform in the master manifest
+	for (const FString& SplitPlatformName : InLocTextHelper.GetPlatformsToSplit())
+	{
+		const FName SplitPlatformFName = *SplitPlatformName;
+		if (!OutPerPlatformLocRes.Contains(SplitPlatformFName))
+		{
+			OutPerPlatformLocRes.Add(SplitPlatformFName, MakeShared<FTextLocalizationResource>());
+		}
+	}
+
 	// Add each manifest entry to the LocRes file
-	InLocTextHelper.EnumerateSourceTexts([&InLocTextHelper, &InCultureToGenerate, &bSkipSourceCheck, &InLocResID, &OutLocRes, InPriority, bIsNativeCulture, Culture](TSharedRef<FManifestEntry> InManifestEntry) -> bool
+	InLocTextHelper.EnumerateSourceTexts([&InLocTextHelper, &InCultureToGenerate, &bSkipSourceCheck, &InLocResID, &OutPlatformAgnosticLocRes, &OutPerPlatformLocRes, InPriority, bIsNativeCulture, Culture](TSharedRef<FManifestEntry> InManifestEntry) -> bool
 	{
 		// For each context, we may need to create a different or even multiple LocRes entries.
 		for (const FManifestContext& Context : InManifestEntry->Contexts)
@@ -54,8 +64,19 @@ bool FTextLocalizationResourceGenerator::GenerateLocRes(const FLocTextHelper& In
 					}
 				}
 
+				// Find the LocRes to update
+				FTextLocalizationResource* LocResToUpdate = &OutPlatformAgnosticLocRes;
+				if (!Context.PlatformName.IsNone())
+				{
+					if (TSharedRef<FTextLocalizationResource>* PerPlatformLocRes = OutPerPlatformLocRes.Find(Context.PlatformName))
+					{
+						LocResToUpdate = &PerPlatformLocRes->Get();
+					}
+				}
+				check(LocResToUpdate);
+
 				// Add this entry to the LocRes
-				OutLocRes.AddEntry(InManifestEntry->Namespace.GetString(), Context.Key.GetString(), InManifestEntry->Source.Text, TranslationText.Text, InPriority, InLocResID);
+				LocResToUpdate->AddEntry(InManifestEntry->Namespace.GetString(), Context.Key.GetString(), InManifestEntry->Source.Text, TranslationText.Text, InPriority, InLocResID);
 			}
 		}
 
@@ -166,6 +187,7 @@ bool FTextLocalizationResourceGenerator::GenerateLocResAndUpdateLiveEntriesFromC
 	}
 
 	FTextLocalizationResource TextLocalizationResource;
+	TMap<FName, TSharedRef<FTextLocalizationResource>> Unused_PerPlatformLocRes;
 	for (int32 CultureIndex = 0; CultureIndex < CulturesToGenerate.Num(); ++CultureIndex)
 	{
 		const FString& CultureName = CulturesToGenerate[CultureIndex];
@@ -173,7 +195,7 @@ bool FTextLocalizationResourceGenerator::GenerateLocResAndUpdateLiveEntriesFromC
 		const FString CulturePath = DestinationPath / CultureName;
 		const FString ResourceFilePath = FPaths::ConvertRelativePathToFull(CulturePath / ResourceName);
 
-		if (!GenerateLocRes(LocTextHelper, CultureName, bSkipSourceCheck, FTextKey(ResourceFilePath), TextLocalizationResource, CultureIndex))
+		if (!GenerateLocRes(LocTextHelper, CultureName, bSkipSourceCheck, FTextKey(ResourceFilePath), TextLocalizationResource, Unused_PerPlatformLocRes, CultureIndex))
 		{
 			UE_LOG(LogTextLocalizationResourceGenerator, Error, TEXT("Failed to generate localization resource for culture '%s'."), *CultureName);
 			return false;
diff --git a/Engine/Source/Developer/Localization/Public/LocTextHelper.h b/Engine/Source/Developer/Localization/Public/LocTextHelper.h
--- a/Engine/Source/Developer/Localization/Public/LocTextHelper.h
+++ b/Engine/Source/Developer/Localization/Public/LocTextHelper.h
@@ -7,6 +7,7 @@
 #include "Internationalization/InternationalizationManifest.h"
 #include "Internationalization/InternationalizationArchive.h"
 #include "Internationalization/LocKeyFuncs.h"
+#include "LocTextHelper.generated.h"
 
 class FLocMetadataObject;
 
@@ -31,6 +32,18 @@ enum class ELocTextExportSourceMethod : uint8
 	NativeText,
 };
 
+/** How should we split platform specific localization data? */
+UENUM()
+enum class ELocTextPlatformSplitMode : uint8
+{
+	/** Don't split platform specific localization data */
+	None,
+	/** Split platform specific localization data for restricted platforms only */
+	Restricted,
+	/** Split platform specific localization data for all platforms */
+	All,
+};
+
 /**
  * Interface for the loc file notify API.
  * This can be used to integrate with services like source control.
@@ -54,6 +67,18 @@ public:
 	virtual void PostFileWrite(const FString& InFilename) = 0;
 };
 
+/**
+ * Utilities to get manage platforms to split.
+ */
+struct LOCALIZATION_API FLocTextPlatformSplitUtils
+{
+	/** True if the given split mode is actually splitting localization data */
+	static bool ShouldSplitPlatformData(const ELocTextPlatformSplitMode& InSplitMode);
+
+	/** Get the platforms names that should be split, based on the given split mode */
+	static const TArray<FString>& GetPlatformsToSplit(const ELocTextPlatformSplitMode& InSplitMode);
+};
+
 /**
  * Class that tracks any conflicts that occur when gathering source text entries.
  */
@@ -220,8 +245,9 @@ public:
 	 * @note This kind of helper is only suitable for dealing with manifests, *not* archives.
 	 *
 	 * @param InLocFileNotifies		Interface for allowing source control integration (may be null).
+	 * @param InPlatformSplitMode	Should we split localization data per-platform?
 	 */
-	explicit FLocTextHelper(TSharedPtr<ILocFileNotifies> InLocFileNotifies);
+	explicit FLocTextHelper(TSharedPtr<ILocFileNotifies> InLocFileNotifies, const ELocTextPlatformSplitMode InPlatformSplitMode = ELocTextPlatformSplitMode::None);
 
 	/**
 	 * Construct a helper for the given target information.
@@ -233,8 +259,24 @@ public:
 	 * @param InNativeCulture		Culture code of the native culture (eg, en), or an empty string if the native culture is unknown.
 	 * @param InForeignCultures		Array of culture codes for the foreign cultures (the native culture will be removed from this array if present).
 	 * @param InLocFileNotifies		Interface for allowing source control integration (may be null).
+	 * @param InPlatformSplitMode	Should we split localization data per-platform?
 	 */
-	FLocTextHelper(FString InTargetPath, FString InManifestName, FString InArchiveName, FString InNativeCulture, TArray<FString> InForeignCultures, TSharedPtr<ILocFileNotifies> InLocFileNotifies);
+	FLocTextHelper(FString InTargetPath, FString InManifestName, FString InArchiveName, FString InNativeCulture, TArray<FString> InForeignCultures, TSharedPtr<ILocFileNotifies> InLocFileNotifies, const ELocTextPlatformSplitMode InPlatformSplitMode = ELocTextPlatformSplitMode::None);
+
+	/**
+	 * Are we splitting localization data per-platform?
+	 */
+	bool ShouldSplitPlatformData() const;
+
+	/**
+	 * How are we splitting localization data per-platform?
+	 */
+	ELocTextPlatformSplitMode GetPlatformSplitMode() const;
+
+	/**
+	 * Get the platforms names that should be split, based on the active split mode.
+	 */
+	const TArray<FString>& GetPlatformsToSplit() const;
 
 	/**
 	 * @return The name of the target we're working with.
@@ -866,6 +908,9 @@ private:
 	 */
 	TSharedPtr<FArchiveEntry> FindTranslationImpl(const FString& InCulture, const FLocKey& InNamespace, const FLocKey& InKey, const TSharedPtr<FLocMetadataObject> InKeyMetadataObj) const;
 
+	/** How are we splitting localization data per-platform? */
+	ELocTextPlatformSplitMode PlatformSplitMode;
+
 	/** The name of the target we're working with */
 	FString TargetName;
 
diff --git a/Engine/Source/Developer/Localization/Public/Serialization/JsonInternationalizationManifestSerializer.h b/Engine/Source/Developer/Localization/Public/Serialization/JsonInternationalizationManifestSerializer.h
--- a/Engine/Source/Developer/Localization/Public/Serialization/JsonInternationalizationManifestSerializer.h
+++ b/Engine/Source/Developer/Localization/Public/Serialization/JsonInternationalizationManifestSerializer.h
@@ -35,7 +35,7 @@ public:
 	 * @param Manifest The populated Internationalization manifest.
 	 * @return true if deserialization was successful, false otherwise.
 	 */
-	static bool DeserializeManifest( const FString& InStr, TSharedRef< FInternationalizationManifest > Manifest );
+	static bool DeserializeManifest( const FString& InStr, TSharedRef< FInternationalizationManifest > Manifest, const FName PlatformName = FName() );
 
 	/**
 	 * Deserializes a Internationalization manifest from a JSON object.
@@ -44,7 +44,7 @@ public:
 	 * @param Manifest The populated Internationalization manifest.
 	 * @return true if deserialization was successful, false otherwise.
 	 */
-	static bool DeserializeManifest( TSharedRef< FJsonObject > InJsonObj, TSharedRef< FInternationalizationManifest > Manifest );
+	static bool DeserializeManifest( TSharedRef< FJsonObject > InJsonObj, TSharedRef< FInternationalizationManifest > Manifest, const FName PlatformName = FName() );
 
 	/**
 	 * Deserializes a Internationalization manifest from a JSON file.
@@ -53,7 +53,7 @@ public:
 	 * @param Manifest The populated Internationalization manifest.
 	 * @return true if deserialization was successful, false otherwise.
 	 */
-	static bool DeserializeManifestFromFile( const FString& InJsonFile, TSharedRef< FInternationalizationManifest > Manifest );
+	static bool DeserializeManifestFromFile( const FString& InJsonFile, TSharedRef< FInternationalizationManifest > Manifest, const FName PlatformName = FName() );
 
 	/**
 	 * Serializes a Internationalization manifest to a JSON string.
@@ -92,7 +92,7 @@ protected:
 	 *
 	 * @return true if deserialization was successful, false otherwise.
 	 */
-	static bool DeserializeInternal( const TSharedRef< FJsonObject > InJsonObj, TSharedRef< FInternationalizationManifest > Manifest );
+	static bool DeserializeInternal( const TSharedRef< FJsonObject > InJsonObj, TSharedRef< FInternationalizationManifest > Manifest, const FName PlatformName );
 
 	/**
 	 * Convert a Internationalization manifest to a JSON object.
@@ -113,7 +113,7 @@ protected:
 	 *
 	 * @return true if successful, false otherwise.
 	 */
-	static bool JsonObjToManifest( TSharedRef< FJsonObject > InJsonObj, FString InNamespace, TSharedRef< FInternationalizationManifest > Manifest );
+	static bool JsonObjToManifest( TSharedRef< FJsonObject > InJsonObj, FString InNamespace, TSharedRef< FInternationalizationManifest > Manifest, const FName PlatformName );
 
 	/**
 	 * Takes a Internationalization manifest and arranges the data into a hierarchy based on namespace.
@@ -151,6 +151,7 @@ public:
 	static const FString TAG_DEPRECATED_DEFAULTTEXT;
 	static const FString TAG_SOURCE;
 	static const FString TAG_SOURCE_TEXT;
+	static const FString TAG_PLATFORM_NAME;
 	static const FString TAG_METADATA;
 	static const FString TAG_METADATA_INFO;
 	static const FString TAG_METADATA_KEY;
diff --git a/Engine/Source/Developer/Localization/Public/TextLocalizationResourceGenerator.h b/Engine/Source/Developer/Localization/Public/TextLocalizationResourceGenerator.h
--- a/Engine/Source/Developer/Localization/Public/TextLocalizationResourceGenerator.h
+++ b/Engine/Source/Developer/Localization/Public/TextLocalizationResourceGenerator.h
@@ -21,7 +21,7 @@ public:
 	/**
 	 * Given a loc text helper, generate a compiled LocRes resource for the given culture.
 	 */
-	LOCALIZATION_API static bool GenerateLocRes(const FLocTextHelper& InLocTextHelper, const FString& InCultureToGenerate, const bool bSkipSourceCheck, const FTextKey& InLocResID, FTextLocalizationResource& OutLocRes, const int32 InPriority = 0);
+	LOCALIZATION_API static bool GenerateLocRes(const FLocTextHelper& InLocTextHelper, const FString& InCultureToGenerate, const bool bSkipSourceCheck, const FTextKey& InLocResID, FTextLocalizationResource& OutPlatformAgnosticLocRes, TMap<FName, TSharedRef<FTextLocalizationResource>>& OutPerPlatformLocRes, const int32 InPriority = 0);
 
 	/**
 	 * Given a config file, generate a compiled LocRes resource for the active culture and use it to update the live-entries in the localization manager.
diff --git a/Engine/Source/Editor/UnrealEd/Classes/Commandlets/GatherTextCommandletBase.h b/Engine/Source/Editor/UnrealEd/Classes/Commandlets/GatherTextCommandletBase.h
--- a/Engine/Source/Editor/UnrealEd/Classes/Commandlets/GatherTextCommandletBase.h
+++ b/Engine/Source/Editor/UnrealEd/Classes/Commandlets/GatherTextCommandletBase.h
@@ -57,7 +57,7 @@ class UNREALED_API UGatherTextCommandletBase : public UCommandlet
 	GENERATED_UCLASS_BODY()
 
 public:
-	virtual void Initialize( const TSharedPtr< FLocTextHelper >& InGatherManifestHelper, const TSharedPtr< FLocalizationSCC >& InSourceControlInfo );
+	virtual void Initialize( const TSharedRef< FLocTextHelper >& InGatherManifestHelper, const TSharedPtr< FLocalizationSCC >& InSourceControlInfo );
 
 	// Wrappers for extracting config values
 	bool GetBoolFromConfig( const TCHAR* Section, const TCHAR* Key, bool& OutValue, const FString& Filename );
@@ -66,11 +66,19 @@ public:
 	int32 GetStringArrayFromConfig( const TCHAR* Section, const TCHAR* Key, TArray<FString>& OutArr, const FString& Filename );
 	int32 GetPathArrayFromConfig( const TCHAR* Section, const TCHAR* Key, TArray<FString>& OutArr, const FString& Filename );
 
+	// Utilities for split platform detection
+	bool IsSplitPlatformName(const FName InPlatformName) const;
+	bool ShouldSplitPlatformForPath(const FString& InPath, FName* OutPlatformName = nullptr) const;
+	FName GetSplitPlatformNameFromPath(const FString& InPath) const;
+
 protected:
 	TSharedPtr< FLocTextHelper > GatherManifestHelper;
 
 	TSharedPtr< FLocalizationSCC > SourceControlInfo;
 
+	/** Mapping from platform name to the path marker for that platform */
+	TMap<FName, FString> SplitPlatforms;
+
 private:
 	virtual void CreateCustomEngine(const FString& Params) override ; //Disallow other text commandlets to make their own engine.	
 };
diff --git a/Engine/Source/Editor/UnrealEd/Classes/Commandlets/GatherTextFromMetadataCommandlet.h b/Engine/Source/Editor/UnrealEd/Classes/Commandlets/GatherTextFromMetadataCommandlet.h
--- a/Engine/Source/Editor/UnrealEd/Classes/Commandlets/GatherTextFromMetadataCommandlet.h
+++ b/Engine/Source/Editor/UnrealEd/Classes/Commandlets/GatherTextFromMetadataCommandlet.h
@@ -30,7 +30,7 @@ private:
 
 private:
 	void GatherTextFromUObjects(const TArray<FString>& IncludePaths, const TArray<FString>& ExcludePaths, const FGatherParameters& Arguments);
-	void GatherTextFromUObject(UField* const Field, const FGatherParameters& Arguments);
+	void GatherTextFromUObject(UField* const Field, const FGatherParameters& Arguments, const FName InPlatformName);
 
 private:
 	bool ShouldGatherFromEditorOnlyData;
diff --git a/Engine/Source/Editor/UnrealEd/Classes/Commandlets/GatherTextFromSourceCommandlet.h b/Engine/Source/Editor/UnrealEd/Classes/Commandlets/GatherTextFromSourceCommandlet.h
--- a/Engine/Source/Editor/UnrealEd/Classes/Commandlets/GatherTextFromSourceCommandlet.h
+++ b/Engine/Source/Editor/UnrealEd/Classes/Commandlets/GatherTextFromSourceCommandlet.h
@@ -22,10 +22,10 @@ class UGatherTextFromSourceCommandlet : public UGatherTextCommandletBase
 private:
 #define LOC_DEFINE_REGION
 
-	enum class EMacroBlockState : uint8
+	enum class EEditorOnlyDefineState : uint8
 	{
-		Normal,
-		EditorOnly,
+		Undefined,
+		Defined,
 	};
 
 	struct FSourceLocation
@@ -55,6 +55,7 @@ private:
 	{
 		FString SourceString;
 		FSourceLocation SourceLocation;
+		FName PlatformName;
 		bool bIsEditorOnly;
 	};
 
@@ -85,7 +86,7 @@ private:
 
 		void FlushMacroStack();
 
-		EMacroBlockState EvaluateMacroStack() const;
+		EEditorOnlyDefineState EvaluateEditorOnlyDefineState() const;
 
 		void SetDefine( const FString& InDefineCtx );
 
@@ -102,6 +103,7 @@ private:
 		//Working data
 		FString Filename;
 		int32 LineNumber;
+		FName FilePlatformName;
 		FString LineText;
 		FString Namespace;
 		bool ExcludedRegion;
@@ -115,15 +117,13 @@ private:
 		//Should editor-only data be included in this gather?
 		bool ShouldGatherFromEditorOnlyData;
 
-		//Destination location of the parsed FLocTextEntrys
-		TSharedPtr< FLocTextHelper > GatherManifestHelper;
-
 		//Discovered string table data from all files
 		TMap<FName, FParsedStringTable> ParsedStringTables;
 
-		FSourceFileParseContext()
+		FSourceFileParseContext(UGatherTextFromSourceCommandlet* InOwnerCommandlet)
 			: Filename()
 			, LineNumber(0)
+			, FilePlatformName()
 			, LineText()
 			, Namespace()
 			, ExcludedRegion(false)
@@ -134,21 +134,23 @@ private:
 			, WithinNamespaceDefine(false)
 			, WithinStartingLine()
 			, ShouldGatherFromEditorOnlyData(false)
-			, GatherManifestHelper()
 			, MacroBlockStack()
-			, CachedMacroBlockState()
+			, CachedEditorOnlyDefineState()
+			, OwnerCommandlet(InOwnerCommandlet)
 		{
-
+			check(OwnerCommandlet);
 		}
 
 	private:
 		bool AddStringTableImpl( const FName InTableId, const FString& InTableNamespace );
-		bool AddStringTableEntryImpl( const FName InTableId, const FString& InKey, const FString& InSourceString, const FSourceLocation& InSourceLocation );
+		bool AddStringTableEntryImpl( const FName InTableId, const FString& InKey, const FString& InSourceString, const FSourceLocation& InSourceLocation, const FName InPlatformName );
 		bool AddStringTableEntryMetaDataImpl( const FName InTableId, const FString& InKey, const FName InMetaDataId, const FString& InMetaData, const FSourceLocation& InSourceLocation );
 
 		//Working data
 		TArray<FString> MacroBlockStack;
-		mutable TOptional<EMacroBlockState> CachedMacroBlockState;
+		mutable TOptional<EEditorOnlyDefineState> CachedEditorOnlyDefineState;
+
+		UGatherTextFromSourceCommandlet* OwnerCommandlet;
 	};
 
 	class FParsableDescriptor
diff --git a/Engine/Source/Editor/UnrealEd/Private/Commandlets/GatherTextCommandlet.cpp b/Engine/Source/Editor/UnrealEd/Private/Commandlets/GatherTextCommandlet.cpp
--- a/Engine/Source/Editor/UnrealEd/Private/Commandlets/GatherTextCommandlet.cpp
+++ b/Engine/Source/Editor/UnrealEd/Private/Commandlets/GatherTextCommandlet.cpp
@@ -1,11 +1,13 @@
 // Copyright 1998-2019 Epic Games, Inc. All Rights Reserved.
 
 #include "Commandlets/GatherTextCommandlet.h"
-#include "UObject/Class.h"
+#include "Misc/App.h"
 #include "Misc/Paths.h"
 #include "Misc/ConfigCacheIni.h"
-#include "Misc/App.h"
+#include "HAL/FileManager.h"
+#include "UObject/Class.h"
 #include "UObject/Package.h"
+#include "SourceControlHelpers.h"
 
 DEFINE_LOG_CATEGORY_STATIC(LogGatherTextCommandlet, Log, All);
 
@@ -122,8 +124,23 @@ int32 UGatherTextCommandlet::ProcessGatherConfig(const FString& GatherTextConfig
 
 	UE_LOG(LogGatherTextCommandlet, Display, TEXT("Beginning GatherText Commandlet for '%s'"), *GatherTextConfigPath);
 
+	// Read in the platform split mode to use
+	ELocTextPlatformSplitMode PlatformSplitMode = ELocTextPlatformSplitMode::None;
+	{
+		FString PlatformSplitModeString;
+		if (GetStringFromConfig(TEXT("CommonSettings"), TEXT("PlatformSplitMode"), PlatformSplitModeString, GatherTextConfigPath))
+		{
+			UEnum* PlatformSplitModeEnum = FindObjectChecked<UEnum>(ANY_PACKAGE, TEXT("ELocTextPlatformSplitMode"));
+			const int64 PlatformSplitModeInt = PlatformSplitModeEnum->GetValueByName(*PlatformSplitModeString);
+			if (PlatformSplitModeInt != INDEX_NONE)
+			{
+				PlatformSplitMode = (ELocTextPlatformSplitMode)PlatformSplitModeInt;
+			}
+		}
+	}
+
 	// Basic helper that can be used only to gather a new manifest for writing
-	TSharedRef<FLocTextHelper> CommandletGatherManifestHelper = MakeShareable(new FLocTextHelper(MakeShareable(new FLocFileSCCNotifies(CommandletSourceControlInfo))));
+	TSharedRef<FLocTextHelper> CommandletGatherManifestHelper = MakeShareable(new FLocTextHelper(MakeShareable(new FLocFileSCCNotifies(CommandletSourceControlInfo)), PlatformSplitMode));
 	CommandletGatherManifestHelper->LoadManifest(ELocTextHelperLoadFlags::Create);
 
 	const FString GatherTextStepPrefix = TEXT("GatherTextStep");
@@ -148,7 +165,7 @@ int32 UGatherTextCommandlet::ProcessGatherConfig(const FString& GatherTextConfig
 		return NumericalSuffixOne < NumericalSuffixTwo;
 	});
 
-	//Execute each step defined in the config file.
+	// Execute each step defined in the config file.
 	for (const FString& StepName : StepNames)
 	{
 		FString CommandletClassName = GConfig->GetStr( *StepName, TEXT("CommandletClass"), GatherTextConfigPath ) + TEXT("Commandlet");
@@ -207,6 +224,58 @@ int32 UGatherTextCommandlet::ProcessGatherConfig(const FString& GatherTextConfig
 		UE_LOG(LogGatherTextCommandlet, Display, TEXT("Completed %s: %s in %f seconds"), *StepName, *CommandletClassName, FPlatformTime::Seconds() - CommandletExecutionStartTime);
 	}
 
+	// Clean-up any stale per-platform data
+	{
+		FString DestinationPath;
+		if (GetPathFromConfig(TEXT("CommonSettings"), TEXT("DestinationPath"), DestinationPath, GatherTextConfigPath))
+		{
+			IFileManager& FileManager = IFileManager::Get();
+
+			auto RemoveDirectory = [&FileManager](const TCHAR* InDirectory)
+			{
+				FileManager.IterateDirectoryRecursively(InDirectory, [&FileManager](const TCHAR* FilenameOrDirectory, bool bIsDirectory) -> bool
+				{
+					if (!bIsDirectory)
+					{
+						if (!USourceControlHelpers::IsAvailable() || !USourceControlHelpers::MarkFileForDelete(FilenameOrDirectory))
+						{
+							FileManager.Delete(FilenameOrDirectory, false, true);
+						}
+					}
+					return true;
+				});
+				FileManager.DeleteDirectory(InDirectory, false, true);
+			};
+
+			const FString PlatformLocalizationPath = DestinationPath / FPaths::GetPlatformLocalizationFolderName();
+			if (CommandletGatherManifestHelper->ShouldSplitPlatformData())
+			{
+				// Remove any stale platform sub-folders
+				FileManager.IterateDirectory(*PlatformLocalizationPath, [&](const TCHAR* FilenameOrDirectory, bool bIsDirectory) -> bool
+				{
+					if (bIsDirectory)
+					{
+						const FString SplitPlatformName = FPaths::GetCleanFilename(FilenameOrDirectory);
+						if (!CommandletGatherManifestHelper->GetPlatformsToSplit().Contains(SplitPlatformName))
+						{
+							RemoveDirectory(FilenameOrDirectory);
+						}
+					}
+					return true;
+				});
+			}
+			else
+			{
+				// Remove the entire Platforms folder
+				RemoveDirectory(*PlatformLocalizationPath);
+			}
+		}
+		else
+		{
+			UE_LOG(LogGatherTextCommandlet, Warning, TEXT("No destination path specified in the 'CommonSettings' section. Cannot check for stale per-platform data!"));
+		}
+	}
+
 	return 0;
 }
 
diff --git a/Engine/Source/Editor/UnrealEd/Private/Commandlets/GatherTextCommandletBase.cpp b/Engine/Source/Editor/UnrealEd/Private/Commandlets/GatherTextCommandletBase.cpp
--- a/Engine/Source/Editor/UnrealEd/Private/Commandlets/GatherTextCommandletBase.cpp
+++ b/Engine/Source/Editor/UnrealEd/Private/Commandlets/GatherTextCommandletBase.cpp
@@ -25,10 +25,21 @@ UGatherTextCommandletBase::UGatherTextCommandletBase(const FObjectInitializer& O
 {
 }
 
-void UGatherTextCommandletBase::Initialize( const TSharedPtr< FLocTextHelper >& InGatherManifestHelper, const TSharedPtr< FLocalizationSCC >& InSourceControlInfo )
+void UGatherTextCommandletBase::Initialize( const TSharedRef< FLocTextHelper >& InGatherManifestHelper, const TSharedPtr< FLocalizationSCC >& InSourceControlInfo )
 {
 	GatherManifestHelper = InGatherManifestHelper;
 	SourceControlInfo = InSourceControlInfo;
+
+	// Cache the split platform info
+	SplitPlatforms.Reset();
+	if (InGatherManifestHelper->ShouldSplitPlatformData())
+	{
+		for (const FString& SplitPlatformName : InGatherManifestHelper->GetPlatformsToSplit())
+		{
+			SplitPlatforms.Add(*SplitPlatformName, FString::Printf(TEXT("/%s/"), *SplitPlatformName));
+		}
+		SplitPlatforms.KeySort(TLess<FName>());
+	}
 }
 
 void UGatherTextCommandletBase::CreateCustomEngine(const FString& Params)
@@ -36,6 +47,33 @@ void UGatherTextCommandletBase::CreateCustomEngine(const FString& Params)
 	GEngine = GEditor = NULL;//Force a basic default engine. 
 }
 
+bool UGatherTextCommandletBase::IsSplitPlatformName(const FName InPlatformName) const
+{
+	return SplitPlatforms.Contains(InPlatformName);
+}
+
+bool UGatherTextCommandletBase::ShouldSplitPlatformForPath(const FString& InPath, FName* OutPlatformName) const
+{
+	const FName SplitPlatformName = GetSplitPlatformNameFromPath(InPath);
+	if (OutPlatformName)
+	{
+		*OutPlatformName = SplitPlatformName;
+	}
+	return !SplitPlatformName.IsNone();
+}
+
+FName UGatherTextCommandletBase::GetSplitPlatformNameFromPath(const FString& InPath) const
+{
+	for (const auto& SplitPlatformsPair : SplitPlatforms)
+	{
+		if (InPath.Contains(SplitPlatformsPair.Value))
+		{
+			return SplitPlatformsPair.Key;
+		}
+	}
+	return FName();
+}
+
 bool UGatherTextCommandletBase::GetBoolFromConfig( const TCHAR* Section, const TCHAR* Key, bool& OutValue, const FString& Filename )
 {
 	bool bSuccess = GConfig->GetBool( Section, Key, OutValue, Filename );
diff --git a/Engine/Source/Editor/UnrealEd/Private/Commandlets/GatherTextFromAssetsCommandlet.cpp b/Engine/Source/Editor/UnrealEd/Private/Commandlets/GatherTextFromAssetsCommandlet.cpp
--- a/Engine/Source/Editor/UnrealEd/Private/Commandlets/GatherTextFromAssetsCommandlet.cpp
+++ b/Engine/Source/Editor/UnrealEd/Private/Commandlets/GatherTextFromAssetsCommandlet.cpp
@@ -222,6 +222,7 @@ void UGatherTextFromAssetsCommandlet::ProcessGatherableTextDataArray(const TArra
 				Context.InfoMetadataObj = !(FLocMetadataObject::IsMetadataExactMatch(&TextSourceSiteContext.InfoMetaData, &DefaultMetadataObject)) ? MakeShareable(new FLocMetadataObject(TextSourceSiteContext.InfoMetaData)) : nullptr;
 				Context.bIsOptional = TextSourceSiteContext.IsOptional;
 				Context.SourceLocation = TextSourceSiteContext.SiteDescription;
+				Context.PlatformName = GetSplitPlatformNameFromPath(TextSourceSiteContext.SiteDescription);
 
 				FLocItem Source(GatherableTextData.SourceData.SourceString);
 
diff --git a/Engine/Source/Editor/UnrealEd/Private/Commandlets/GatherTextFromMetadataCommandlet.cpp b/Engine/Source/Editor/UnrealEd/Private/Commandlets/GatherTextFromMetadataCommandlet.cpp
--- a/Engine/Source/Editor/UnrealEd/Private/Commandlets/GatherTextFromMetadataCommandlet.cpp
+++ b/Engine/Source/Editor/UnrealEd/Private/Commandlets/GatherTextFromMetadataCommandlet.cpp
@@ -161,11 +161,12 @@ void UGatherTextFromMetaDataCommandlet::GatherTextFromUObjects(const TArray<FStr
 			continue;
 		}
 
-		GatherTextFromUObject(*It, Arguments);
+		const FName MetaDataPlatformName = GetSplitPlatformNameFromPath(SourceFilePath);
+		GatherTextFromUObject(*It, Arguments, MetaDataPlatformName);
 	}
 }
 
-void UGatherTextFromMetaDataCommandlet::GatherTextFromUObject(UField* const Field, const FGatherParameters& Arguments)
+void UGatherTextFromMetaDataCommandlet::GatherTextFromUObject(UField* const Field, const FGatherParameters& Arguments, const FName InPlatformName)
 {
 	// Gather for object.
 	{
@@ -191,6 +192,7 @@ void UGatherTextFromMetaDataCommandlet::GatherTextFromUObject(UField* const Fiel
 					FManifestContext Context;
 					Context.Key = FText::Format(Arguments.OutputKeys[i], PatternArguments).ToString();
 					Context.SourceLocation = FString::Printf(TEXT("From metadata for key %s of member %s in %s"), *Arguments.InputKeys[i], *Field->GetName(), *Field->GetFullGroupName(true));
+					Context.PlatformName = InPlatformName;
 					GatherManifestHelper->AddSourceText(Namespace, LocItem, Context);
 				}
 			}
@@ -227,6 +229,7 @@ void UGatherTextFromMetaDataCommandlet::GatherTextFromUObject(UField* const Fiel
 							FManifestContext Context;
 							Context.Key = FText::Format(Arguments.OutputKeys[j], PatternArguments).ToString();
 							Context.SourceLocation = FString::Printf(TEXT("From metadata for key %s of enum value %s of enum %s in %s"), *Arguments.InputKeys[j], *Enum->GetNameStringByIndex(i), *Enum->GetName(), *Enum->GetFullGroupName(true));
+							Context.PlatformName = InPlatformName;
 							GatherManifestHelper->AddSourceText(Namespace, LocItem, Context);
 						}
 					}
diff --git a/Engine/Source/Editor/UnrealEd/Private/Commandlets/GatherTextFromSourceCommandlet.cpp b/Engine/Source/Editor/UnrealEd/Private/Commandlets/GatherTextFromSourceCommandlet.cpp
--- a/Engine/Source/Editor/UnrealEd/Private/Commandlets/GatherTextFromSourceCommandlet.cpp
+++ b/Engine/Source/Editor/UnrealEd/Private/Commandlets/GatherTextFromSourceCommandlet.cpp
@@ -262,8 +262,7 @@ int32 UGatherTextFromSourceCommandlet::Main( const FString& Params )
 	Parsables.Add(new FIniNamespaceDescriptor());
 
 	// Init a parse context to track the state of the file parsing 
-	FSourceFileParseContext ParseCtxt;
-	ParseCtxt.GatherManifestHelper = GatherManifestHelper;
+	FSourceFileParseContext ParseCtxt(this);
 
 	// Get whether we should gather editor-only data. Typically only useful for the localization of UE4 itself.
 	if (!GetBoolFromConfig(*SectionName, TEXT("ShouldGatherFromEditorOnlyData"), ParseCtxt.ShouldGatherFromEditorOnlyData, GatherTextConfigPath))
@@ -287,6 +286,7 @@ int32 UGatherTextFromSourceCommandlet::Main( const FString& Params )
 		ParseCtxt.Filename = SourceFile;
 		FPaths::MakePathRelativeTo(ParseCtxt.Filename, *ProjectBasePath);
 		ParseCtxt.LineNumber = 0;
+		ParseCtxt.FilePlatformName = GetSplitPlatformNameFromPath(ParseCtxt.Filename);
 		ParseCtxt.LineText.Empty();
 		ParseCtxt.Namespace.Empty();
 		ParseCtxt.ExcludedRegion = false;
@@ -334,6 +334,7 @@ int32 UGatherTextFromSourceCommandlet::Main( const FString& Params )
 					FManifestContext SourceContext;
 					SourceContext.Key = ParsedStringTableEntryPair.Key;
 					SourceContext.SourceLocation = ParsedStringTableEntryPair.Value.SourceLocation.ToString();
+					SourceContext.PlatformName = ParsedStringTableEntryPair.Value.PlatformName;
 
 					const FParsedStringTableEntryMetaDataMap* ParsedMetaDataMap = ParsedStringTablePair.Value.MetaDataEntries.Find(ParsedStringTableEntryPair.Key);
 					if (ParsedMetaDataMap && ParsedMetaDataMap->Num() > 0)
@@ -353,7 +354,7 @@ int32 UGatherTextFromSourceCommandlet::Main( const FString& Params )
 						*ParsedStringTableEntryPair.Value.SourceLocation.ToString()
 						);
 
-					ParseCtxt.GatherManifestHelper->AddSourceText(ParsedStringTablePair.Value.TableNamespace, FLocItem(ParsedStringTableEntryPair.Value.SourceString), SourceContext, &SourceDescription);
+					GatherManifestHelper->AddSourceText(ParsedStringTablePair.Value.TableNamespace, FLocItem(ParsedStringTableEntryPair.Value.SourceString), SourceContext, &SourceDescription);
 				}
 			}
 		}
@@ -929,7 +930,7 @@ bool UGatherTextFromSourceCommandlet::ParseSourceText(const FString& Text, const
 
 bool UGatherTextFromSourceCommandlet::FSourceFileParseContext::AddManifestText( const FString& Token, const FString& InNamespace, const FString& SourceText, const FManifestContext& Context )
 {
-	const bool bIsEditorOnly = EvaluateMacroStack() == EMacroBlockState::EditorOnly;
+	const bool bIsEditorOnly = EvaluateEditorOnlyDefineState() == EEditorOnlyDefineState::Defined;
 
 	if (!bIsEditorOnly || ShouldGatherFromEditorOnlyData)
 	{
@@ -938,7 +939,7 @@ bool UGatherTextFromSourceCommandlet::FSourceFileParseContext::AddManifestText(
 			*Filename, 
 			LineNumber, 
 			*LineText);
-		return GatherManifestHelper->AddSourceText(InNamespace, FLocItem(SourceText), Context, &EntryDescription);
+		return OwnerCommandlet->GatherManifestHelper->AddSourceText(InNamespace, FLocItem(SourceText), Context, &EntryDescription);
 	}
 
 	return false;
@@ -947,7 +948,7 @@ bool UGatherTextFromSourceCommandlet::FSourceFileParseContext::AddManifestText(
 void UGatherTextFromSourceCommandlet::FSourceFileParseContext::PushMacroBlock( const FString& InBlockCtx )
 {
 	MacroBlockStack.Push(InBlockCtx);
-	CachedMacroBlockState.Reset();
+	CachedEditorOnlyDefineState.Reset();
 }
 
 void UGatherTextFromSourceCommandlet::FSourceFileParseContext::PopMacroBlock()
@@ -955,36 +956,37 @@ void UGatherTextFromSourceCommandlet::FSourceFileParseContext::PopMacroBlock()
 	if (MacroBlockStack.Num() > 0)
 	{
 		MacroBlockStack.Pop(/*bAllowShrinking*/false);
-		CachedMacroBlockState.Reset();
+		CachedEditorOnlyDefineState.Reset();
 	}
 }
 
 void UGatherTextFromSourceCommandlet::FSourceFileParseContext::FlushMacroStack()
 {
 	MacroBlockStack.Reset();
+	CachedEditorOnlyDefineState.Reset();
 }
 
-UGatherTextFromSourceCommandlet::EMacroBlockState UGatherTextFromSourceCommandlet::FSourceFileParseContext::EvaluateMacroStack() const
+UGatherTextFromSourceCommandlet::EEditorOnlyDefineState UGatherTextFromSourceCommandlet::FSourceFileParseContext::EvaluateEditorOnlyDefineState() const
 {
-	if (CachedMacroBlockState.IsSet())
+	if (CachedEditorOnlyDefineState.IsSet())
 	{
-		return CachedMacroBlockState.GetValue();
+		return CachedEditorOnlyDefineState.GetValue();
 	}
 
 	static const FString WithEditorString = TEXT("WITH_EDITOR");
 	static const FString WithEditorOnlyDataString = TEXT("WITH_EDITORONLY_DATA");
 
-	CachedMacroBlockState = EMacroBlockState::Normal;
+	CachedEditorOnlyDefineState = EEditorOnlyDefineState::Undefined;
 	for (const FString& BlockCtx : MacroBlockStack)
 	{
 		if (BlockCtx.Equals(WithEditorString, ESearchCase::CaseSensitive) || BlockCtx.Equals(WithEditorOnlyDataString, ESearchCase::CaseSensitive))
 		{
-			CachedMacroBlockState = EMacroBlockState::EditorOnly;
+			CachedEditorOnlyDefineState = EEditorOnlyDefineState::Defined;
 			break;
 		}
 	}
 
-	return CachedMacroBlockState.GetValue();
+	return CachedEditorOnlyDefineState.GetValue();
 }
 
 void UGatherTextFromSourceCommandlet::FSourceFileParseContext::SetDefine(const FString& InDefineCtx)
@@ -1085,9 +1087,9 @@ bool UGatherTextFromSourceCommandlet::FSourceFileParseContext::AddStringTableImp
 	return true;
 }
 
-bool UGatherTextFromSourceCommandlet::FSourceFileParseContext::AddStringTableEntryImpl(const FName InTableId, const FString& InKey, const FString& InSourceString, const FSourceLocation& InSourceLocation)
+bool UGatherTextFromSourceCommandlet::FSourceFileParseContext::AddStringTableEntryImpl(const FName InTableId, const FString& InKey, const FString& InSourceString, const FSourceLocation& InSourceLocation, const FName InPlatformName)
 {
-	const bool bIsEditorOnly = EvaluateMacroStack() == EMacroBlockState::EditorOnly;
+	const bool bIsEditorOnly = EvaluateEditorOnlyDefineState() == EEditorOnlyDefineState::Defined;
 
 	// String table entries may be parsed before the string table itself (due to code ordering), so we may need to add our string table below
 	FParsedStringTable& ParsedStringTable = ParsedStringTables.FindOrAdd(InTableId);
@@ -1111,6 +1113,7 @@ bool UGatherTextFromSourceCommandlet::FSourceFileParseContext::AddStringTableEnt
 		FParsedStringTableEntry& ParsedStringTableEntry = ParsedStringTable.TableEntries.Add(InKey);
 		ParsedStringTableEntry.SourceString = InSourceString;
 		ParsedStringTableEntry.SourceLocation = InSourceLocation;
+		ParsedStringTableEntry.PlatformName = InPlatformName;
 		ParsedStringTableEntry.bIsEditorOnly = bIsEditorOnly;
 		return true;
 	}
@@ -1118,7 +1121,7 @@ bool UGatherTextFromSourceCommandlet::FSourceFileParseContext::AddStringTableEnt
 
 bool UGatherTextFromSourceCommandlet::FSourceFileParseContext::AddStringTableEntryMetaDataImpl(const FName InTableId, const FString& InKey, const FName InMetaDataId, const FString& InMetaData, const FSourceLocation& InSourceLocation)
 {
-	const bool bIsEditorOnly = EvaluateMacroStack() == EMacroBlockState::EditorOnly;
+	const bool bIsEditorOnly = EvaluateEditorOnlyDefineState() == EEditorOnlyDefineState::Defined;
 
 	// String table meta-data may be parsed before the string table itself (due to code ordering), so we may need to add our string table below
 	FParsedStringTable& ParsedStringTable = ParsedStringTables.FindOrAdd(InTableId);
@@ -1162,10 +1165,12 @@ void UGatherTextFromSourceCommandlet::FSourceFileParseContext::AddStringTableFro
 		FStringTableRef TmpStringTable = FStringTable::NewStringTable();
 		if (TmpStringTable->ImportStrings(FullImportPath))
 		{
+			const FSourceLocation SourceLocation = FSourceLocation(InTableFilename, INDEX_NONE);
+			const FName TablePlatformName = OwnerCommandlet->GetSplitPlatformNameFromPath(InTableFilename);
+
 			TmpStringTable->EnumerateSourceStrings([&](const FString& InKey, const FString& InSourceString)
 			{
-				const FSourceLocation SourceLocation = FSourceLocation(InTableFilename, INDEX_NONE);
-				AddStringTableEntryImpl(InTableId, InKey, InSourceString, SourceLocation);
+				AddStringTableEntryImpl(InTableId, InKey, InSourceString, SourceLocation, TablePlatformName);
 
 				TmpStringTable->EnumerateMetaData(InKey, [&](const FName InMetaDataId, const FString& InMetaData)
 				{
@@ -1185,7 +1190,7 @@ void UGatherTextFromSourceCommandlet::FSourceFileParseContext::AddStringTableFro
 
 void UGatherTextFromSourceCommandlet::FSourceFileParseContext::AddStringTableEntry(const FName InTableId, const FString& InKey, const FString& InSourceString)
 {
-	AddStringTableEntryImpl(InTableId, InKey, InSourceString, FSourceLocation(Filename, LineNumber));
+	AddStringTableEntryImpl(InTableId, InKey, InSourceString, FSourceLocation(Filename, LineNumber), FilePlatformName);
 }
 
 void UGatherTextFromSourceCommandlet::FSourceFileParseContext::AddStringTableEntryMetaData(const FName InTableId, const FString& InKey, const FName InMetaDataId, const FString& InMetaData)
@@ -1453,6 +1458,7 @@ void UGatherTextFromSourceCommandlet::FUICommandMacroDescriptor::TryParseArgs(co
 			FManifestContext CommandContext;
 			CommandContext.Key = Identifier;
 			CommandContext.SourceLocation = SourceLocation;
+			CommandContext.PlatformName = Context.FilePlatformName;
 
 			Context.AddManifestText(GetToken(), Namespace, SourceText, CommandContext);
 
@@ -1468,6 +1474,7 @@ void UGatherTextFromSourceCommandlet::FUICommandMacroDescriptor::TryParseArgs(co
 					FManifestContext CommandTooltipContext;
 					CommandTooltipContext.Key = Identifier + TEXT("_ToolTip");
 					CommandTooltipContext.SourceLocation = SourceLocation;
+					CommandTooltipContext.PlatformName = CommandContext.PlatformName;
 
 					Context.AddManifestText(GetToken(), Namespace, TooltipSourceText, CommandTooltipContext);
 				}
@@ -1602,6 +1609,7 @@ void UGatherTextFromSourceCommandlet::FStringMacroDescriptor::TryParse(const FSt
 					FManifestContext MacroContext;
 					MacroContext.Key = Identifier;
 					MacroContext.SourceLocation = SourceLocation;
+					MacroContext.PlatformName = Context.FilePlatformName;
 
 					Context.AddManifestText( GetToken(), Namespace.Get(FString()), SourceText, MacroContext );
 				}
diff --git a/Engine/Source/Editor/UnrealEd/Private/Commandlets/GenerateGatherArchiveCommandlet.cpp b/Engine/Source/Editor/UnrealEd/Private/Commandlets/GenerateGatherArchiveCommandlet.cpp
--- a/Engine/Source/Editor/UnrealEd/Private/Commandlets/GenerateGatherArchiveCommandlet.cpp
+++ b/Engine/Source/Editor/UnrealEd/Private/Commandlets/GenerateGatherArchiveCommandlet.cpp
@@ -203,13 +203,13 @@ int32 UGenerateGatherArchiveCommandlet::Main( const FString& Params )
 	}
 
 	// Load the manifest and all archives
-	FLocTextHelper LocTextHelper(DestinationPath, ManifestName, ArchiveName, NativeCulture, CulturesToGenerate, MakeShareable(new FLocFileSCCNotifies(SourceControlInfo)));
+	FLocTextHelper LocTextHelper(DestinationPath, ManifestName, ArchiveName, NativeCulture, CulturesToGenerate, GatherManifestHelper->GetLocFileNotifies(), GatherManifestHelper->GetPlatformSplitMode());
 	{
 		FText LoadError;
 		if (!LocTextHelper.LoadAll(ELocTextHelperLoadFlags::LoadOrCreate, &LoadError))
 		{
 			UE_LOG(LogGenerateArchiveCommandlet, Error, TEXT("%s"), *LoadError.ToString());
-			return false;
+			return -1;
 		}
 	}
 
diff --git a/Engine/Source/Editor/UnrealEd/Private/Commandlets/GenerateGatherManifestCommandlet.cpp b/Engine/Source/Editor/UnrealEd/Private/Commandlets/GenerateGatherManifestCommandlet.cpp
--- a/Engine/Source/Editor/UnrealEd/Private/Commandlets/GenerateGatherManifestCommandlet.cpp
+++ b/Engine/Source/Editor/UnrealEd/Private/Commandlets/GenerateGatherManifestCommandlet.cpp
@@ -74,7 +74,7 @@ int32 UGenerateGatherManifestCommandlet::Main( const FString& Params )
 		FText OutError;
 		if (!GatherManifestHelper->AddDependency(ManifestDependency, &OutError))
 		{
-			UE_LOG(LogGenerateManifestCommandlet, Error, TEXT("The GenerateGatherManifest commandlet couldn't load the specified manifest dependency: '%'. %s"), *ManifestDependency, *OutError.ToString());
+			UE_LOG(LogGenerateManifestCommandlet, Error, TEXT("Failed to add a manifest dependency: %s"), *OutError.ToString());
 			return -1;
 		}
 	}
@@ -86,7 +86,7 @@ int32 UGenerateGatherManifestCommandlet::Main( const FString& Params )
 	FText ManifestSaveError;
 	if (!GatherManifestHelper->SaveManifest(ManifestPath, &ManifestSaveError))
 	{
-		UE_LOG(LogGenerateManifestCommandlet, Error,TEXT("Failed to write manifest to %s. %s."), *ManifestPath, *ManifestSaveError.ToString());
+		UE_LOG(LogGenerateManifestCommandlet, Error, TEXT("%s"), *ManifestSaveError.ToString());
 		return -1;
 	}
 
diff --git a/Engine/Source/Editor/UnrealEd/Private/Commandlets/GenerateTextLocalizationResourceCommandlet.cpp b/Engine/Source/Editor/UnrealEd/Private/Commandlets/GenerateTextLocalizationResourceCommandlet.cpp
--- a/Engine/Source/Editor/UnrealEd/Private/Commandlets/GenerateTextLocalizationResourceCommandlet.cpp
+++ b/Engine/Source/Editor/UnrealEd/Private/Commandlets/GenerateTextLocalizationResourceCommandlet.cpp
@@ -125,13 +125,13 @@ int32 UGenerateTextLocalizationResourceCommandlet::Main(const FString& Params)
 	}
 
 	// Load the manifest and all archives
-	FLocTextHelper LocTextHelper(SourcePath, ManifestName, ArchiveName, NativeCultureName, CulturesToGenerate, MakeShareable(new FLocFileSCCNotifies(SourceControlInfo)));
+	FLocTextHelper LocTextHelper(SourcePath, ManifestName, ArchiveName, NativeCultureName, CulturesToGenerate, GatherManifestHelper->GetLocFileNotifies(), GatherManifestHelper->GetPlatformSplitMode());
 	{
 		FText LoadError;
 		if (!LocTextHelper.LoadAll(ELocTextHelperLoadFlags::LoadOrCreate, &LoadError))
 		{
 			UE_LOG(LogGenerateTextLocalizationResourceCommandlet, Error, TEXT("%s"), *LoadError.ToString());
-			return false;
+			return -1;
 		}
 	}
 
@@ -148,26 +148,57 @@ int32 UGenerateTextLocalizationResourceCommandlet::Main(const FString& Params)
 		if (!bLocMetaFileSaved)
 		{
 			UE_LOG(LogGenerateTextLocalizationResourceCommandlet, Error, TEXT("Could not write file %s"), *TextLocalizationMetaDataResourcePath);
-			return false;
+			return -1;
 		}
 	}
 
 	// Generate the LocRes file for each culture
 	for (const FString& CultureName : CulturesToGenerate)
 	{
-		const FString TextLocalizationResourcePath = DestinationPath / CultureName / ResourceName;
-
-		const bool bLocResFileSaved = FLocalizedAssetSCCUtil::SaveFileWithSCC(SourceControlInfo, TextLocalizationResourcePath, [&LocTextHelper, &CultureName, &bSkipSourceCheck](const FString& InSaveFileName) -> bool
+		auto GenerateSingleLocRes = [this, &DestinationPath, &CultureName, &ResourceName](const FTextLocalizationResource& InLocRes, const FName InPlatformName) -> bool
 		{
-			FTextLocalizationResource LocRes;
-			return FTextLocalizationResourceGenerator::GenerateLocRes(LocTextHelper, CultureName, bSkipSourceCheck, FTextKey(InSaveFileName), LocRes) && LocRes.SaveToFile(InSaveFileName);
-		});
-
-		if (!bLocResFileSaved)
+			FString TextLocalizationResourcePath;
+			if (InPlatformName.IsNone())
+			{
+				TextLocalizationResourcePath = DestinationPath / CultureName / ResourceName;
+			}
+			else
+			{
+				TextLocalizationResourcePath = DestinationPath / FPaths::GetPlatformLocalizationFolderName() / InPlatformName.ToString() / CultureName / ResourceName;
+			}
+
+			const bool bLocResFileSaved = FLocalizedAssetSCCUtil::SaveFileWithSCC(SourceControlInfo, TextLocalizationResourcePath, [&InLocRes](const FString& InSaveFileName) -> bool
+			{
+				return InLocRes.SaveToFile(InSaveFileName);
+			});
+
+			if (!bLocResFileSaved)
+			{
+				UE_LOG(LogGenerateTextLocalizationResourceCommandlet, Error, TEXT("Could not write file %s"), *TextLocalizationResourcePath);
+				return false;
+			}
+
+			return true;
+		};
+
+		FTextLocalizationResource PlatformAgnosticLocRes;
+		TMap<FName, TSharedRef<FTextLocalizationResource>> PerPlatformLocRes;
+		const FTextKey LocResId = DestinationPath / CultureName / ResourceName;
+		if (!FTextLocalizationResourceGenerator::GenerateLocRes(LocTextHelper, CultureName, bSkipSourceCheck, LocResId, PlatformAgnosticLocRes, PerPlatformLocRes))
 		{
-			UE_LOG(LogGenerateTextLocalizationResourceCommandlet, Error, TEXT("Could not write file %s"), *TextLocalizationResourcePath);
+			UE_LOG(LogGenerateTextLocalizationResourceCommandlet, Error, TEXT("Failed to generate LocRes %s"), LocResId.GetChars());
 			return false;
 		}
+	
+		bool bSuccess = GenerateSingleLocRes(PlatformAgnosticLocRes, FName());
+		for (const auto& PerPlatformLocResPair : PerPlatformLocRes)
+		{
+			bSuccess &= GenerateSingleLocRes(*PerPlatformLocResPair.Value, PerPlatformLocResPair.Key);
+		}
+		if (!bSuccess)
+		{
+			return -1;
+		}
 	}
 
 	return 0;
diff --git a/Engine/Source/Editor/UnrealEd/Private/Commandlets/InternationalizationExportCommandlet.cpp b/Engine/Source/Editor/UnrealEd/Private/Commandlets/InternationalizationExportCommandlet.cpp
--- a/Engine/Source/Editor/UnrealEd/Private/Commandlets/InternationalizationExportCommandlet.cpp
+++ b/Engine/Source/Editor/UnrealEd/Private/Commandlets/InternationalizationExportCommandlet.cpp
@@ -135,7 +135,7 @@ int32 UInternationalizationExportCommandlet::Main(const FString& Params)
 	if (bDoImport)
 	{
 		// Load the manifest and all archives
-		FLocTextHelper LocTextHelper(DestinationPath, ManifestName, ArchiveName, NativeCultureName, CulturesToGenerate, MakeShareable(new FLocFileSCCNotifies(SourceControlInfo)));
+		FLocTextHelper LocTextHelper(DestinationPath, ManifestName, ArchiveName, NativeCultureName, CulturesToGenerate, GatherManifestHelper->GetLocFileNotifies(), GatherManifestHelper->GetPlatformSplitMode());
 		{
 			FText LoadError;
 			if (!LocTextHelper.LoadAll(ELocTextHelperLoadFlags::LoadOrCreate, &LoadError))
@@ -159,7 +159,7 @@ int32 UInternationalizationExportCommandlet::Main(const FString& Params)
 		GetBoolFromConfig(*SectionName, TEXT("ShouldPersistCommentsOnExport"), bShouldPersistComments, ConfigPath);
 
 		// Load the manifest and all archives
-		FLocTextHelper LocTextHelper(SourcePath, ManifestName, ArchiveName, NativeCultureName, CulturesToGenerate, MakeShareable(new FLocFileSCCNotifies(SourceControlInfo)));
+		FLocTextHelper LocTextHelper(SourcePath, ManifestName, ArchiveName, NativeCultureName, CulturesToGenerate, GatherManifestHelper->GetLocFileNotifies(), GatherManifestHelper->GetPlatformSplitMode());
 		{
 			FText LoadError;
 			if (!LocTextHelper.LoadAll(ELocTextHelperLoadFlags::LoadOrCreate, &LoadError))
diff --git a/Engine/Source/Programs/AutomationTool/Localization/LocalizationProvider.cs b/Engine/Source/Programs/AutomationTool/Localization/LocalizationProvider.cs
--- a/Engine/Source/Programs/AutomationTool/Localization/LocalizationProvider.cs
+++ b/Engine/Source/Programs/AutomationTool/Localization/LocalizationProvider.cs
@@ -4,6 +4,7 @@ using System.Collections.Generic;
 using System.IO;
 using AutomationTool;
 using UnrealBuildTool;
+using Tools.DotNETCommon;
 
 namespace EpicGames.Localization
 {
@@ -30,6 +31,30 @@ public struct ProjectImportExportInfo
 
 	/** True if we should use a per-culture directly when importing/exporting */
 	public bool bUseCultureDirectory;
+
+	/** The platform names that we have split into sub-folders */
+	public List<string> SplitPlatformNames;
+
+	/** The platforms sub-folder */
+	public const string PlatformLocalizationFolderName = "Platforms";
+
+	public void CalculateSplitPlatformNames(string RootWorkingDirectory)
+	{
+		// Is this isn't a single culture import/export, then also check for split platform sub-folders
+		SplitPlatformNames = new List<string>();
+		if (bUseCultureDirectory)
+		{
+			var PlatformSourceDirectory = new DirectoryReference(CommandUtils.CombinePaths(RootWorkingDirectory, DestinationPath, PlatformLocalizationFolderName));
+			if (DirectoryReference.Exists(PlatformSourceDirectory))
+			{
+				foreach (DirectoryReference FoundDirectory in DirectoryReference.EnumerateDirectories(PlatformSourceDirectory, "*", SearchOption.TopDirectoryOnly))
+				{
+					string SplitPlatformName = CommandUtils.GetLastDirectoryName(FoundDirectory.FullName);
+					SplitPlatformNames.Add(SplitPlatformName);
+				}
+			}
+		}
+	}
 };
 
 public struct ProjectStepInfo
diff --git a/Engine/Source/Programs/AutomationTool/OneSkyLocalization/OneSkyLocalizationProvider.cs b/Engine/Source/Programs/AutomationTool/OneSkyLocalization/OneSkyLocalizationProvider.cs
--- a/Engine/Source/Programs/AutomationTool/OneSkyLocalization/OneSkyLocalizationProvider.cs
+++ b/Engine/Source/Programs/AutomationTool/OneSkyLocalization/OneSkyLocalizationProvider.cs
@@ -48,8 +48,17 @@ public class OneSkyLocalizationProvider : LocalizationProvider
 
 	public override void DownloadProjectFromLocalizationProvider(string ProjectName, ProjectImportExportInfo ProjectImportInfo)
 	{
-		var OneSkyFileName = GetOneSkyFilename(ProjectImportInfo.PortableObjectName);
 		var OneSkyProject = GetOneSkyProject(ProjectName);
+		DownloadLatestPOFiles(OneSkyProject, null, ProjectImportInfo);
+		foreach (var Platform in ProjectImportInfo.SplitPlatformNames)
+		{
+			DownloadLatestPOFiles(OneSkyProject, Platform, ProjectImportInfo);
+		}
+	}
+
+	private void DownloadLatestPOFiles(Project OneSkyProject, string Platform, ProjectImportExportInfo ProjectImportInfo)
+	{
+		var OneSkyFileName = GetOneSkyFilename(ProjectImportInfo.PortableObjectName, Platform);
 		var OneSkyFile = OneSkyProject.UploadedFiles.FirstOrDefault(f => f.Filename == OneSkyFileName);
 
 		// Export
@@ -71,7 +80,10 @@ public class OneSkyLocalizationProvider : LocalizationProvider
 				}
 			}
 
-			ExportOneSkyFileToDirectory(OneSkyFile, new DirectoryInfo(CommandUtils.CombinePaths(RootWorkingDirectory, ProjectImportInfo.DestinationPath)), ProjectImportInfo.PortableObjectName, CulturesToExport, ProjectImportInfo.bUseCultureDirectory);
+			var DestinationDirectory = String.IsNullOrEmpty(Platform) 
+					? new DirectoryInfo(CommandUtils.CombinePaths(RootWorkingDirectory, ProjectImportInfo.DestinationPath))
+					: new DirectoryInfo(CommandUtils.CombinePaths(RootWorkingDirectory, ProjectImportInfo.DestinationPath, ProjectImportExportInfo.PlatformLocalizationFolderName, Platform));
+			ExportOneSkyFileToDirectory(OneSkyFile, DestinationDirectory, ProjectImportInfo.PortableObjectName, CulturesToExport, ProjectImportInfo.bUseCultureDirectory);
 		}
 	}
 
@@ -186,20 +198,12 @@ public class OneSkyLocalizationProvider : LocalizationProvider
 	{
 		var OneSkyProject = GetOneSkyProject(ProjectName);
 
-		Func<string, FileInfo> GetPathForCulture = (string Culture) =>
-		{
-			if (ProjectExportInfo.bUseCultureDirectory)
-			{
-				return new FileInfo(Path.Combine(RootWorkingDirectory, ProjectExportInfo.DestinationPath, Culture, ProjectExportInfo.PortableObjectName));
-			}
-			else
-			{
-				return new FileInfo(Path.Combine(RootWorkingDirectory, ProjectExportInfo.DestinationPath, ProjectExportInfo.PortableObjectName));
-			}
-		};
-
 		// Upload the .po file for the native culture first
-		UploadFileToOneSky(OneSkyProject, GetPathForCulture(ProjectExportInfo.NativeCulture), ProjectExportInfo.NativeCulture);
+		UploadFileToOneSky(OneSkyProject, ProjectExportInfo.NativeCulture, null, ProjectExportInfo);
+		foreach (var Platform in ProjectExportInfo.SplitPlatformNames)
+		{
+			UploadFileToOneSky(OneSkyProject, ProjectExportInfo.NativeCulture, Platform, ProjectExportInfo);
+		}
 
 		if (bUploadAllCultures)
 		{
@@ -209,7 +213,11 @@ public class OneSkyLocalizationProvider : LocalizationProvider
 				// Skip native culture as we uploaded it above
 				if (Culture != ProjectExportInfo.NativeCulture)
 				{
-					UploadFileToOneSky(OneSkyProject, GetPathForCulture(Culture), Culture);
+					UploadFileToOneSky(OneSkyProject, Culture, null, ProjectExportInfo);
+					foreach (var Platform in ProjectExportInfo.SplitPlatformNames)
+					{
+						UploadFileToOneSky(OneSkyProject, Culture, Platform, ProjectExportInfo);
+					}
 				}
 			}
 		}
@@ -244,8 +252,14 @@ public class OneSkyLocalizationProvider : LocalizationProvider
 		return null;
 	}
 
-	private void UploadFileToOneSky(OneSky.Project OneSkyProject, FileInfo FileToUpload, string Culture)
+	private void UploadFileToOneSky(OneSky.Project OneSkyProject, string Culture, string Platform, ProjectImportExportInfo ProjectExportInfo)
 	{
+		var SourceDirectory = String.IsNullOrEmpty(Platform)
+				? new DirectoryInfo(CommandUtils.CombinePaths(RootWorkingDirectory, ProjectExportInfo.DestinationPath))
+				: new DirectoryInfo(CommandUtils.CombinePaths(RootWorkingDirectory, ProjectExportInfo.DestinationPath, ProjectImportExportInfo.PlatformLocalizationFolderName, Platform));
+		var CultureDirectory = (ProjectExportInfo.bUseCultureDirectory) ? new DirectoryInfo(Path.Combine(SourceDirectory.FullName, Culture)) : SourceDirectory;
+
+		var FileToUpload = new FileInfo(Path.Combine(CultureDirectory.FullName, ProjectExportInfo.PortableObjectName));
 		using (var FileStream = FileToUpload.OpenRead())
 		{
 			// Read the BOM
@@ -258,7 +272,7 @@ public class OneSkyLocalizationProvider : LocalizationProvider
 				FileStream.Position = 0;
 			}
 
-			var OneSkyFileName = GetOneSkyFilename(Path.GetFileName(FileToUpload.FullName));
+			var OneSkyFileName = GetOneSkyFilename(ProjectExportInfo.PortableObjectName, Platform);
 
 			Console.WriteLine("Uploading: '{0}' as '{1}' ({2})", FileToUpload.FullName, OneSkyFileName, Culture);
 
@@ -303,9 +317,15 @@ public class OneSkyLocalizationProvider : LocalizationProvider
 		return OneSkyProject;
 	}
 
-	private string GetOneSkyFilename(string BaseFilename)
+	private string GetOneSkyFilename(string BaseFilename, string Platform)
 	{
 		var OneSkyFileName = BaseFilename;
+		if (!String.IsNullOrEmpty(Platform))
+		{
+			// Apply the platform suffix. XLoc will take care of merging the files from different platforms together.
+			var OneSkyFileNameWithSuffix = Path.GetFileNameWithoutExtension(OneSkyFileName) + "_" + Platform + Path.GetExtension(OneSkyFileName);
+			OneSkyFileName = OneSkyFileNameWithSuffix;
+		}
 		if (!String.IsNullOrEmpty(LocalizationBranchName))
 		{
 			// Apply the branch suffix. OneSky will take care of merging the files from different branches together.
diff --git a/Engine/Source/Programs/AutomationTool/Scripts/CopyBuildToStagingDirectory.Automation.cs b/Engine/Source/Programs/AutomationTool/Scripts/CopyBuildToStagingDirectory.Automation.cs
--- a/Engine/Source/Programs/AutomationTool/Scripts/CopyBuildToStagingDirectory.Automation.cs
+++ b/Engine/Source/Programs/AutomationTool/Scripts/CopyBuildToStagingDirectory.Automation.cs
@@ -268,13 +268,25 @@ public partial class Project : CommandUtils
 
 	private static void StageLocalizationDataForTarget(DeploymentContext SC, List<string> CulturesToStage, DirectoryReference SourceDirectory)
 	{
+		var PlatformSourceDirectory = new DirectoryReference(CombinePaths(SourceDirectory.FullName, "Platforms", ConfigHierarchy.GetIniPlatformName(SC.StageTargetPlatform.IniPlatformType)));
+		if (!DirectoryReference.Exists(PlatformSourceDirectory))
+		{
+			PlatformSourceDirectory = null;
+		}
+
 		foreach (FileReference SourceFile in DirectoryReference.EnumerateFiles(SourceDirectory, "*.locmeta"))
 		{
 			SC.StageFile(StagedFileType.UFS, SourceFile);
 		}
+
 		foreach (string Culture in CulturesToStage)
 		{
 			StageLocalizationDataForCulture(SC, Culture, SourceDirectory);
+
+			if (PlatformSourceDirectory != null)
+			{
+				StageLocalizationDataForCulture(SC, Culture, PlatformSourceDirectory);
+			}
 		}
 	}
 
diff --git a/Engine/Source/Programs/AutomationTool/Scripts/Localisation.Automation.cs b/Engine/Source/Programs/AutomationTool/Scripts/Localisation.Automation.cs
--- a/Engine/Source/Programs/AutomationTool/Scripts/Localisation.Automation.cs
+++ b/Engine/Source/Programs/AutomationTool/Scripts/Localisation.Automation.cs
@@ -1,4 +1,4 @@
-﻿// Copyright 1998-2019 Epic Games, Inc. All Rights Reserved.
+// Copyright 1998-2019 Epic Games, Inc. All Rights Reserved.
 using System;
 using System.Collections.Generic;
 using System.Globalization;
@@ -340,6 +340,8 @@ class Localize : BuildCommand
 				// Upload all text to our localization provider
 				foreach (var ProjectInfo in ProjectInfos)
 				{
+					// Recalculate the split platform paths before doing the upload, as the export may have changed them
+					ProjectInfo.ExportInfo.CalculateSplitPlatformNames(RootLocalizationTargetDirectory);
 					LocProvider.UploadProjectToLocalizationProvider(ProjectInfo.ProjectName, ProjectInfo.ExportInfo);
 				}
 			}
@@ -361,7 +363,7 @@ class Localize : BuildCommand
 		{
 			ProjectInfo.LocalizationSteps.Add(new ProjectStepInfo("Monolithic", MonolithicConfigFile));
 
-			ProjectInfo.ImportInfo = GenerateProjectImportExportInfo(MonolithicConfigFile);
+			ProjectInfo.ImportInfo = GenerateProjectImportExportInfo(RootWorkingDirectory, MonolithicConfigFile);
 			ProjectInfo.ExportInfo = ProjectInfo.ImportInfo;
 		}
 		else
@@ -384,11 +386,11 @@ class Localize : BuildCommand
 
 					if (FileSuffix.Suffix == "Import")
 					{
-						ProjectInfo.ImportInfo = GenerateProjectImportExportInfo(ModularConfigFile);
+						ProjectInfo.ImportInfo = GenerateProjectImportExportInfo(RootWorkingDirectory, ModularConfigFile);
 					}
 					else if (FileSuffix.Suffix == "Export")
 					{
-						ProjectInfo.ExportInfo = GenerateProjectImportExportInfo(ModularConfigFile);
+						ProjectInfo.ExportInfo = GenerateProjectImportExportInfo(RootWorkingDirectory, ModularConfigFile);
 					}
 				}
 				else if (FileSuffix.Required)
@@ -401,7 +403,7 @@ class Localize : BuildCommand
 		return ProjectInfo;
 	}
 
-	private ProjectImportExportInfo GenerateProjectImportExportInfo(string LocalizationConfigFile)
+	private ProjectImportExportInfo GenerateProjectImportExportInfo(string RootWorkingDirectory, string LocalizationConfigFile)
 	{
 		var ProjectImportExportInfo = new ProjectImportExportInfo();
 
@@ -444,6 +446,8 @@ class Localize : BuildCommand
 			ProjectImportExportInfo.bUseCultureDirectory = true;
 		}
 
+		ProjectImportExportInfo.CalculateSplitPlatformNames(RootWorkingDirectory);
+
 		return ProjectImportExportInfo;
 	}
 
diff --git a/Engine/Source/Programs/AutomationTool/XLocLocalization/XLocLocalizationProvider.cs b/Engine/Source/Programs/AutomationTool/XLocLocalization/XLocLocalizationProvider.cs
--- a/Engine/Source/Programs/AutomationTool/XLocLocalization/XLocLocalizationProvider.cs
+++ b/Engine/Source/Programs/AutomationTool/XLocLocalization/XLocLocalizationProvider.cs
@@ -76,7 +76,11 @@ public abstract class XLocLocalizationProvider : LocalizationProvider
 					continue;
 				}
 
-				DownloadLatestPOFile(XLocApiClient, AuthToken, Culture, ProjectImportInfo);
+				DownloadLatestPOFile(XLocApiClient, AuthToken, Culture, null, ProjectImportInfo);
+				foreach (var Platform in ProjectImportInfo.SplitPlatformNames)
+				{
+					DownloadLatestPOFile(XLocApiClient, AuthToken, Culture, Platform, ProjectImportInfo);
+				}
 			}
 		}
 		finally
@@ -85,9 +89,9 @@ public abstract class XLocLocalizationProvider : LocalizationProvider
 		}
 	}
 
-	private void DownloadLatestPOFile(XLocApiClient XLocApiClient, string AuthToken, string Culture, ProjectImportExportInfo ProjectImportInfo)
+	private void DownloadLatestPOFile(XLocApiClient XLocApiClient, string AuthToken, string Culture, string Platform, ProjectImportExportInfo ProjectImportInfo)
 	{
-		var XLocFilename = GetXLocFilename(ProjectImportInfo.PortableObjectName);
+		var XLocFilename = GetXLocFilename(ProjectImportInfo.PortableObjectName, Platform);
 
 		// This will throw if the requested culture is invalid, but we don't want to let that kill the whole gather
 		var LatestBuildXml = "";
@@ -151,7 +155,9 @@ public abstract class XLocLocalizationProvider : LocalizationProvider
 		}
 		else
 		{
-			var DestinationDirectory = new DirectoryInfo(CommandUtils.CombinePaths(RootWorkingDirectory, ProjectImportInfo.DestinationPath));
+			var DestinationDirectory = String.IsNullOrEmpty(Platform) 
+					? new DirectoryInfo(CommandUtils.CombinePaths(RootWorkingDirectory, ProjectImportInfo.DestinationPath))
+					: new DirectoryInfo(CommandUtils.CombinePaths(RootWorkingDirectory, ProjectImportInfo.DestinationPath, ProjectImportExportInfo.PlatformLocalizationFolderName, Platform));
 			var CultureDirectory = (ProjectImportInfo.bUseCultureDirectory) ? new DirectoryInfo(Path.Combine(DestinationDirectory.FullName, Culture)) : DestinationDirectory;
 			if (!CultureDirectory.Exists)
 			{
@@ -234,7 +240,11 @@ public abstract class XLocLocalizationProvider : LocalizationProvider
 			var AuthToken = RequestAuthTokenWithRetry(XLocApiClient);
 
 			// Upload the .po file for the native culture first
-			UploadLatestPOFile(TransferServiceClient, AuthToken, ProjectExportInfo.NativeCulture, ProjectExportInfo);
+			UploadLatestPOFile(TransferServiceClient, AuthToken, ProjectExportInfo.NativeCulture, null, ProjectExportInfo);
+			foreach (var Platform in ProjectExportInfo.SplitPlatformNames)
+			{
+				UploadLatestPOFile(TransferServiceClient, AuthToken, ProjectExportInfo.NativeCulture, Platform, ProjectExportInfo);
+			}
 
 			if (bUploadAllCultures)
 			{
@@ -244,7 +254,11 @@ public abstract class XLocLocalizationProvider : LocalizationProvider
 					// Skip native culture as we uploaded it above
 					if (Culture != ProjectExportInfo.NativeCulture)
 					{
-						UploadLatestPOFile(TransferServiceClient, AuthToken, Culture, ProjectExportInfo);
+						UploadLatestPOFile(TransferServiceClient, AuthToken, Culture, null, ProjectExportInfo);
+						foreach (var Platform in ProjectExportInfo.SplitPlatformNames)
+						{
+							UploadLatestPOFile(TransferServiceClient, AuthToken, Culture, Platform, ProjectExportInfo);
+						}
 					}
 				}
 			}
@@ -256,13 +270,15 @@ public abstract class XLocLocalizationProvider : LocalizationProvider
 		}
 	}
 
-	private void UploadLatestPOFile(TransferServiceClient TransferServiceClient, string AuthToken, string Culture, ProjectImportExportInfo ProjectExportInfo)
+	private void UploadLatestPOFile(TransferServiceClient TransferServiceClient, string AuthToken, string Culture, string Platform, ProjectImportExportInfo ProjectExportInfo)
 	{
-		var SourceDirectory = new DirectoryInfo(CommandUtils.CombinePaths(RootWorkingDirectory, ProjectExportInfo.DestinationPath));
+		var SourceDirectory = String.IsNullOrEmpty(Platform)
+				? new DirectoryInfo(CommandUtils.CombinePaths(RootWorkingDirectory, ProjectExportInfo.DestinationPath))
+				: new DirectoryInfo(CommandUtils.CombinePaths(RootWorkingDirectory, ProjectExportInfo.DestinationPath, ProjectImportExportInfo.PlatformLocalizationFolderName, Platform));
 		var CultureDirectory = (ProjectExportInfo.bUseCultureDirectory) ? new DirectoryInfo(Path.Combine(SourceDirectory.FullName, Culture)) : SourceDirectory;
 
 		var FileToUpload = new FileInfo(Path.Combine(CultureDirectory.FullName, ProjectExportInfo.PortableObjectName));
-		var XLocFilename = GetXLocFilename(ProjectExportInfo.PortableObjectName);
+		var XLocFilename = GetXLocFilename(ProjectExportInfo.PortableObjectName, Platform);
 
 		using (var FileStream = FileToUpload.OpenRead())
 		{
@@ -368,9 +384,15 @@ public abstract class XLocLocalizationProvider : LocalizationProvider
 		return XLocApiClient.GetLatestBuildByFile(Config.APIKey, AuthToken, XLocUtils.MD5HashString(Config.APIKey + Config.APISecret + Config.LocalizationId + LanguageId + RemoteFilename), Config.LocalizationId, LanguageId, RemoteFilename);
 	}
 
-	private string GetXLocFilename(string BaseFilename)
+	private string GetXLocFilename(string BaseFilename, string Platform)
 	{
 		var XLocFilename = BaseFilename;
+		if (!String.IsNullOrEmpty(Platform))
+		{
+			// Apply the platform suffix. XLoc will take care of merging the files from different platforms together.
+			var XLocFilenameWithSuffix = Path.GetFileNameWithoutExtension(XLocFilename) + "_" + Platform + Path.GetExtension(XLocFilename);
+			XLocFilename = XLocFilenameWithSuffix;
+		}
 		if (!String.IsNullOrEmpty(LocalizationBranchName))
 		{
 			// Apply the branch suffix. XLoc will take care of merging the files from different branches together.
diff --git a/Engine/Source/Runtime/Core/Private/Internationalization/InternationalizationManifest.cpp b/Engine/Source/Runtime/Core/Private/Internationalization/InternationalizationManifest.cpp
--- a/Engine/Source/Runtime/Core/Private/Internationalization/InternationalizationManifest.cpp
+++ b/Engine/Source/Runtime/Core/Private/Internationalization/InternationalizationManifest.cpp
@@ -8,6 +8,7 @@ DEFINE_LOG_CATEGORY_STATIC(LogInternationalizationManifestObject, Log, All);
 FManifestContext::FManifestContext(const FManifestContext& Other)
 	: Key(Other.Key)
 	, SourceLocation(Other.SourceLocation)
+	, PlatformName(Other.PlatformName)
 	, bIsOptional(Other.bIsOptional)
 {
 	if (Other.InfoMetadataObj.IsValid())
@@ -27,6 +28,7 @@ FManifestContext& FManifestContext::operator=(const FManifestContext& Other)
 	{
 		Key = Other.Key;
 		SourceLocation = Other.SourceLocation;
+		PlatformName = Other.PlatformName;
 		bIsOptional = Other.bIsOptional;
 		InfoMetadataObj.Reset();
 		KeyMetadataObj.Reset();
@@ -294,7 +296,9 @@ const FManifestContext* FManifestEntry::FindContext(const FLocKey& ContextKey, c
 			{
 				continue;
 			}
-			else if ((!Context.KeyMetadataObj.IsValid() && !KeyMetadata.IsValid()) || (*Context.KeyMetadataObj == *KeyMetadata))
+			
+			// If we get here, the meta-data objects are either both null, or both valid
+			if (!Context.KeyMetadataObj.IsValid() || *Context.KeyMetadataObj == *KeyMetadata)
 			{
 				return &Context;
 			}
@@ -316,3 +320,35 @@ const FManifestContext* FManifestEntry::FindContextByKey(const FLocKey& ContextK
 
 	return nullptr;
 }
+
+void FManifestEntry::MergeContextPlatformInfo(const FManifestContext& InContext)
+{
+	for (FManifestContext& Context : Contexts)
+	{
+		if (Context.Key == InContext.Key)
+		{
+			if (Context.KeyMetadataObj.IsValid() != InContext.KeyMetadataObj.IsValid())
+			{
+				continue;
+			}
+			
+			// If we get here, the meta-data objects are either both null, or both valid
+			if (!Context.KeyMetadataObj.IsValid() || *Context.KeyMetadataObj == *InContext.KeyMetadataObj)
+			{
+				// If the platform name on this context doesn't match what we're being asked to merge to it, clear the 
+				// platform name so that the text becomes platform agnostic (as it is being used by multiple platforms)
+				// Also clear the source location to redact the info
+				if (!Context.PlatformName.IsNone() && Context.PlatformName != InContext.PlatformName)
+				{
+					Context.PlatformName = FName();
+					Context.SourceLocation.Reset();
+				}
+				if (InContext.PlatformName.IsNone() && Context.SourceLocation.IsEmpty())
+				{
+					// Previously redacted source location - use the agnostic source instead
+					Context.SourceLocation = InContext.SourceLocation;
+				}
+			}
+		}
+	}
+}
diff --git a/Engine/Source/Runtime/Core/Private/Internationalization/LocalizationResourceTextSource.cpp b/Engine/Source/Runtime/Core/Private/Internationalization/LocalizationResourceTextSource.cpp
--- a/Engine/Source/Runtime/Core/Private/Internationalization/LocalizationResourceTextSource.cpp
+++ b/Engine/Source/Runtime/Core/Private/Internationalization/LocalizationResourceTextSource.cpp
@@ -2,6 +2,7 @@
 
 #include "Internationalization/LocalizationResourceTextSource.h"
 #include "Internationalization/TextLocalizationResource.h"
+#include "HAL/PlatformProperties.h"
 #include "HAL/FileManager.h"
 #include "Misc/Paths.h"
 #include "Misc/CoreDelegates.h"
@@ -114,6 +115,20 @@ void FLocalizationResourceTextSource::LoadLocalizedResources(const ELocalization
 		}
 	}
 
+	static const FString PlatformLocalizationFolderName = FPaths::GetPlatformLocalizationFolderName();
+	static const FString PlatformName = ANSI_TO_TCHAR(FPlatformProperties::IniPlatformName());
+	auto LoadLocalizationResourcesForCulture = [](FTextLocalizationResource& InOutLocRes, const FString& InLocalizationPath, const FString& InCulture, const FString& InLocResFilename, const int32 InLocResPriority)
+	{
+		const FString PlatformAgnosticLocResFilename = InLocalizationPath / InCulture / InLocResFilename;
+		InOutLocRes.LoadFromFile(PlatformAgnosticLocResFilename, InLocResPriority);
+
+		const FString PlatformSpecificLocResFilename = InLocalizationPath / PlatformLocalizationFolderName / PlatformName / InCulture / InLocResFilename;
+		if (FPaths::FileExists(PlatformSpecificLocResFilename))
+		{
+			InOutLocRes.LoadFromFile(PlatformSpecificLocResFilename, InLocResPriority);
+		}
+	};
+
 	// Load the native texts first to ensure we always apply translations to a consistent base
 	if (PrioritizedNativePaths.Num() > 0)
 	{
@@ -132,7 +147,7 @@ void FLocalizationResourceTextSource::LoadLocalizedResources(const ELocalization
 				// We skip loading the native text if we're transitioning to the native culture as there's no extra work that needs to be done
 				if (!InPrioritizedCultures.Contains(LocMetaResource.NativeCulture))
 				{
-					InOutNativeResource.LoadFromFile(LocalizationPath / LocMetaResource.NativeLocRes, BaseResourcePriority);
+					LoadLocalizationResourcesForCulture(InOutNativeResource, LocalizationPath, LocMetaResource.NativeCulture, FPaths::GetCleanFilename(LocMetaResource.NativeLocRes), BaseResourcePriority);
 				}
 			}
 		}
@@ -152,7 +167,7 @@ void FLocalizationResourceTextSource::LoadLocalizedResources(const ELocalization
 				}
 
 				const FString LocResFilename = FPaths::GetBaseFilename(LocalizationPath) + TEXT(".locres");
-				InOutLocalizedResource.LoadFromFile(LocalizationPath / NativeGameCulture / LocResFilename, BaseResourcePriority);
+				LoadLocalizationResourcesForCulture(InOutLocalizedResource, LocalizationPath, NativeGameCulture, LocResFilename, BaseResourcePriority);
 			}
 		}
 	}
@@ -171,7 +186,7 @@ void FLocalizationResourceTextSource::LoadLocalizedResources(const ELocalization
 				}
 
 				const FString LocResFilename = FPaths::GetBaseFilename(LocalizationPath) + TEXT(".locres");
-				InOutLocalizedResource.LoadFromFile(LocalizationPath / PrioritizedCultureName / LocResFilename, BaseResourcePriority + CultureIndex);
+				LoadLocalizationResourcesForCulture(InOutLocalizedResource, LocalizationPath, PrioritizedCultureName, LocResFilename, BaseResourcePriority + CultureIndex);
 			}
 		}
 	}
diff --git a/Engine/Source/Runtime/Core/Private/Internationalization/TextLocalizationResource.cpp b/Engine/Source/Runtime/Core/Private/Internationalization/TextLocalizationResource.cpp
--- a/Engine/Source/Runtime/Core/Private/Internationalization/TextLocalizationResource.cpp
+++ b/Engine/Source/Runtime/Core/Private/Internationalization/TextLocalizationResource.cpp
@@ -357,7 +357,7 @@ bool FTextLocalizationResource::LoadFromArchive(FArchive& Archive, const FTextKe
 	return true;
 }
 
-bool FTextLocalizationResource::SaveToFile(const FString& FilePath)
+bool FTextLocalizationResource::SaveToFile(const FString& FilePath) const
 {
 	TUniquePtr<FArchive> Writer(IFileManager::Get().CreateFileWriter(*FilePath));
 	if (!Writer)
@@ -371,7 +371,7 @@ bool FTextLocalizationResource::SaveToFile(const FString& FilePath)
 	return bSaved;
 }
 
-bool FTextLocalizationResource::SaveToArchive(FArchive& Archive, const FTextKey& LocResID)
+bool FTextLocalizationResource::SaveToArchive(FArchive& Archive, const FTextKey& LocResID) const
 {
 	// Write the header
 	{
@@ -608,35 +608,21 @@ TArray<FString> TextLocalizationResourceUtil::GetLocalizedCultureNames(const TAr
 {
 	TArray<FString> CultureNames;
 
-	// Find all unique culture folders that exist in the given paths
+	// Find all unique culture folders that exist in the given paths, skipping the platforms sub-folder
+	const FString PlatformFolderName = FPaths::GetPlatformLocalizationFolderName();
 	for (const FString& LocalizationPath : InLocalizationPaths)
 	{
-		/* Visitor class used to enumerate directories of culture */
-		class FCultureEnumeratorVistor : public IPlatformFile::FDirectoryVisitor
+		IFileManager::Get().IterateDirectory(*LocalizationPath, [&CultureNames, &PlatformFolderName](const TCHAR* FilenameOrDirectory, bool bIsDirectory) -> bool
 		{
-		public:
-			FCultureEnumeratorVistor(TArray<FString>& OutCultureNames)
-				: CultureNamesRef(OutCultureNames)
+			if (bIsDirectory && FCString::Stricmp(FilenameOrDirectory, *PlatformFolderName) != 0)
 			{
+				// UE localization resource folders use "en-US" style while ICU uses "en_US"
+				const FString LocalizationFolder = FPaths::GetCleanFilename(FilenameOrDirectory);
+				const FString CanonicalName = FCulture::GetCanonicalName(LocalizationFolder);
+				CultureNames.AddUnique(CanonicalName);
 			}
-
-			virtual bool Visit(const TCHAR* FilenameOrDirectory, bool bIsDirectory) override
-			{
-				if (bIsDirectory)
-				{
-					// UE localization resource folders use "en-US" style while ICU uses "en_US"
-					const FString LocalizationFolder = FPaths::GetCleanFilename(FilenameOrDirectory);
-					const FString CanonicalName = FCulture::GetCanonicalName(LocalizationFolder);
-					CultureNamesRef.AddUnique(CanonicalName);
-				}
-				return true;
-			}
-
-			TArray<FString>& CultureNamesRef;
-		};
-
-		FCultureEnumeratorVistor CultureEnumeratorVistor(CultureNames);
-		IFileManager::Get().IterateDirectory(*LocalizationPath, CultureEnumeratorVistor);
+			return true;
+		});
 	}
 
 	// Remove any cultures that were explicitly disallowed
diff --git a/Engine/Source/Runtime/Core/Private/Misc/Paths.cpp b/Engine/Source/Runtime/Core/Private/Misc/Paths.cpp
--- a/Engine/Source/Runtime/Core/Private/Misc/Paths.cpp
+++ b/Engine/Source/Runtime/Core/Private/Misc/Paths.cpp
@@ -549,6 +549,12 @@ const TArray<FString>& FPaths::GetGameLocalizationPaths()
 	return Results;
 }
 
+FString FPaths::GetPlatformLocalizationFolderName()
+{
+	// Note: If you change this, also update StageLocalizationDataForTarget (CopyBuildToStagingDirectory.Automation.cs), ProjectImportExportInfo.PlatformLocalizationFolderName (LocalizationProvider.cs)
+	return TEXT("Platforms");
+}
+
 const TArray<FString>& FPaths::GetRestrictedFolderNames()
 {
 	static bool bBuiltArray = false;
diff --git a/Engine/Source/Runtime/Core/Public/Internationalization/InternationalizationManifest.h b/Engine/Source/Runtime/Core/Public/Internationalization/InternationalizationManifest.h
--- a/Engine/Source/Runtime/Core/Public/Internationalization/InternationalizationManifest.h
+++ b/Engine/Source/Runtime/Core/Public/Internationalization/InternationalizationManifest.h
@@ -19,6 +19,7 @@ public:
 	FManifestContext()
 		: Key()
 		, SourceLocation()
+		, PlatformName()
 		, bIsOptional(false)
 	{
 	}
@@ -26,6 +27,7 @@ public:
 	explicit FManifestContext(const FLocKey& InKey)
 		: Key(InKey)
 		, SourceLocation()
+		, PlatformName()
 		, bIsOptional(false)
 	{
 	}
@@ -41,6 +43,7 @@ public:
 public:
 	FLocKey Key;
 	FString SourceLocation;
+	FName PlatformName;
 	bool bIsOptional;
 
 	TSharedPtr<FLocMetadataObject> InfoMetadataObj;
@@ -98,6 +101,7 @@ public:
 
 	const FManifestContext* FindContext(const FLocKey& ContextKey, const TSharedPtr<FLocMetadataObject>& KeyMetadata = nullptr) const;
 	const FManifestContext* FindContextByKey(const FLocKey& ContextKey) const;
+	void MergeContextPlatformInfo(const FManifestContext& InContext);
 
 	const FLocKey Namespace;
 	const FLocItem Source;
diff --git a/Engine/Source/Runtime/Core/Public/Internationalization/TextLocalizationResource.h b/Engine/Source/Runtime/Core/Public/Internationalization/TextLocalizationResource.h
--- a/Engine/Source/Runtime/Core/Public/Internationalization/TextLocalizationResource.h
+++ b/Engine/Source/Runtime/Core/Public/Internationalization/TextLocalizationResource.h
@@ -88,10 +88,10 @@ public:
 	bool LoadFromArchive(FArchive& Archive, const FTextKey& LocResID, const int32 Priority);
 
 	/** Save this resource to the given LocRes file. */
-	bool SaveToFile(const FString& FilePath);
+	bool SaveToFile(const FString& FilePath) const;
 
 	/** Save this resource to the given LocRes archive. */
-	bool SaveToArchive(FArchive& Archive, const FTextKey& LocResID);
+	bool SaveToArchive(FArchive& Archive, const FTextKey& LocResID) const;
 
 private:
 	/** Test whether the new entry should replace the current entry, optionally logging a conflict */
diff --git a/Engine/Source/Runtime/Core/Public/Misc/Paths.h b/Engine/Source/Runtime/Core/Public/Misc/Paths.h
--- a/Engine/Source/Runtime/Core/Public/Misc/Paths.h
+++ b/Engine/Source/Runtime/Core/Public/Misc/Paths.h
@@ -327,6 +327,11 @@ public:
 	 */
 	static const TArray<FString>& GetGameLocalizationPaths();
 
+	/**
+	 * Get the name of the platform-specific localization sub-folder
+	 */
+	static FString GetPlatformLocalizationFolderName();
+
 	/** 
 	 * Returns a list of restricted/internal folder names (without any slashes) which may be tested against full paths to determine if a path is restricted or not.
 	 */
