Fixed compile errors when compiling Win64 Clang Editor config due to FLocTextHelper move ctor and assignment.

FLocTextConflicts ConflictTracker member is explicitly marked non-copyable / movable, therefore default move for FLocTextHelper can't be generated. There is currently no reason for FLocTextHelper to be movable, therefore default move ops are now simply removed from code. Move ops can be added to FLocTextConflicts if necessary in the future.

diff --git a/Engine/Source/Developer/Localization/Public/LocTextHelper.h b/Engine/Source/Developer/Localization/Public/LocTextHelper.h
--- a/Engine/Source/Developer/Localization/Public/LocTextHelper.h
+++ b/Engine/Source/Developer/Localization/Public/LocTextHelper.h
@@ -843,10 +843,6 @@ public:
 	 */
 	static bool FindKeysForLegacyTranslation(const TSharedRef<const FInternationalizationManifest>& InManifest, const TSharedPtr<const FInternationalizationArchive>& InNativeArchive, const FLocKey& InNamespace, const FString& InSource, const TSharedPtr<FLocMetadataObject> InKeyMetadataObj, TArray<FLocKey>& OutKeys);
 
-	/** Publicly movable */
-	FLocTextHelper(FLocTextHelper&&) = default;
-	FLocTextHelper& operator=(FLocTextHelper&&) = default;
-
 private:
 	/** Non-copyable */
 	FLocTextHelper(const FLocTextHelper&) = delete;
