Remove duplicate function definition
#rnx

diff --git a/Engine/Source/Developer/Localization/Public/LocTextHelper.h b/Engine/Source/Developer/Localization/Public/LocTextHelper.h
--- a/Engine/Source/Developer/Localization/Public/LocTextHelper.h
+++ b/Engine/Source/Developer/Localization/Public/LocTextHelper.h
@@ -113,10 +113,6 @@ private:
 public:
 	FLocTextConflicts() {}
 
-	/** Publicly movable */
-	FLocTextConflicts(FLocTextConflicts&&) = default;
-	FLocTextConflicts& operator=(FLocTextConflicts&&) = default;
-	
 	/**
 	 * Add a new conflict entry.
 	 *
