PR #6316: Persist translator comments when exporting PO (Contributed by [REDACTED_PERSON])

diff --git a/Engine/Source/Developer/Localization/Private/PortableObjectFormatDOM.cpp b/Engine/Source/Developer/Localization/Private/PortableObjectFormatDOM.cpp
--- a/Engine/Source/Developer/Localization/Private/PortableObjectFormatDOM.cpp
+++ b/Engine/Source/Developer/Localization/Private/PortableObjectFormatDOM.cpp
@@ -906,14 +906,6 @@ void FPortableObjectEntry::AddExtractedComment( const FString& InComment )
 	{
 		ExtractedComments.AddUnique(InComment);
 	}
-
-	//// Extracted comments can contain multiple references in a single line so we parse those out.
-	//TArray<FString> CommentsToProcess;
-	//InComment.ParseIntoArray( CommentsToProcess, TEXT(" "), true );
-	//for( const FString& ExtractedComment : CommentsToProcess )
-	//{
-	//	ExtractedComments.AddUnique( ExtractedComment );
-	//}
 }
 
 void FPortableObjectEntry::AddReference( const FString& InReference )
@@ -922,14 +914,6 @@ void FPortableObjectEntry::AddReference( const FString& InReference )
 	{
 		ReferenceComments.AddUnique(InReference);
 	}
-	
-	//// Reference comments can contain multiple references in a single line so we parse those out.
-	//TArray<FString> ReferencesToProcess;
-	//InReference.ParseIntoArray( ReferencesToProcess, TEXT(" "), true );
-	//for( const FString& Reference : ReferencesToProcess )
-	//{
-	//	ReferenceComments.AddUnique( Reference );
-	//}
 }
 
 void FPortableObjectEntry::AddExtractedComments( const TArray<FString>& InComments )
@@ -948,6 +932,22 @@ void FPortableObjectEntry::AddReferences( const TArray<FString>& InReferences )
 	}
 }
 
+void FPortableObjectEntry::AddTranslatorComment(const FString& InComment)
+{
+	if (!InComment.IsEmpty())
+	{
+		TranslatorComments.AddUnique(InComment);
+	}
+}
+
+void FPortableObjectEntry::AddTranslatorComments(const TArray<FString>& InComments)
+{
+	for (const FString& TranslatorComment : InComments)
+	{
+		AddTranslatorComment(TranslatorComment);
+	}
+}
+
 FString FPortableObjectEntry::ToString() const
 {
 	check( !MsgId.IsEmpty() );
diff --git a/Engine/Source/Developer/Localization/Private/PortableObjectPipeline.cpp b/Engine/Source/Developer/Localization/Private/PortableObjectPipeline.cpp
--- a/Engine/Source/Developer/Localization/Private/PortableObjectPipeline.cpp
+++ b/Engine/Source/Developer/Localization/Private/PortableObjectPipeline.cpp
@@ -51,6 +51,18 @@ namespace
 		FLocKeyPairMultiMap CollapsedNSSourceStringToExpandedNSKey;
 	};
 
+	struct FLocComments
+	{
+		FLocComments(TArray<FString> InExtracted, TArray<FString> InTranslator)
+			: ExtractedComments(MoveTemp(InExtracted))
+			, TranslatorComments(MoveTemp(InTranslator))
+		{
+		}
+
+		TArray<FString> ExtractedComments;
+		TArray<FString> TranslatorComments;
+	};
+
 	void BuildCollapsedManifest(FLocTextHelper& InLocTextHelper, const ELocalizedTextCollapseMode InTextCollapseMode, FCollapsedData& OutCollapsedData, TSharedPtr<FInternationalizationManifest>& OutPlatformAgnosticManifest, TMap<FName, TSharedRef<FInternationalizationManifest>>& OutPerPlatformManifests)
 	{
 		// Always add the split platforms so that they generate an empty manifest if there are no entries for that platform in the master manifest
@@ -154,9 +166,9 @@ namespace
 		}, true);
 	}
 
-	TMap<FPortableObjectEntryKey, TArray<FString>> ExtractPreservedPOComments(const FPortableObjectFormatDOM& InPortableObject)
+	TMap<FPortableObjectEntryKey, FLocComments> ExtractPreservedPOComments(const FPortableObjectFormatDOM& InPortableObject)
 	{
-		TMap<FPortableObjectEntryKey, TArray<FString>> POEntryToCommentMap;
+		TMap<FPortableObjectEntryKey, FLocComments> POEntryToCommentMap;
 		for (auto EntryPairIterator = InPortableObject.GetEntriesIterator(); EntryPairIterator; ++EntryPairIterator)
 		{
 			const TSharedPtr< FPortableObjectEntry >& Entry = EntryPairIterator->Value;
@@ -167,9 +179,9 @@ namespace
 				return !ExtractedComment.StartsWith(TEXT("Key:"), ESearchCase::CaseSensitive) && !ExtractedComment.StartsWith(TEXT("SourceLocation:"), ESearchCase::CaseSensitive) && !ExtractedComment.StartsWith(TEXT("InfoMetaData:"), ESearchCase::CaseSensitive);
 			});
 
-			if (CommentsToPreserve.Num())
+			if (CommentsToPreserve.Num() || Entry->TranslatorComments.Num())
 			{
-				POEntryToCommentMap.Add(FPortableObjectEntryKey(Entry->MsgId, Entry->MsgIdPlural, Entry->MsgCtxt), CommentsToPreserve);
+				POEntryToCommentMap.Add(FPortableObjectEntryKey(Entry->MsgId, Entry->MsgIdPlural, Entry->MsgCtxt), FLocComments(CommentsToPreserve, Entry->TranslatorComments));
 			}
 		}
 		return POEntryToCommentMap;
@@ -376,7 +388,7 @@ namespace
 		if (bShouldPersistComments)
 		{
 			// Preserve comments from the specified file now
-			TMap<FPortableObjectEntryKey, TArray<FString>> POEntryToCommentMap;
+			TMap<FPortableObjectEntryKey, FLocComments> POEntryToCommentMap;
 			{
 				FPortableObjectFormatDOM ExistingPortableObject;
 				if (LoadPOFile(InPOFilePath, ExistingPortableObject))
@@ -391,7 +403,8 @@ namespace
 				const TSharedPtr<FPortableObjectEntry> FoundEntry = NewPortableObject.FindEntry(Pair.Key.MsgId, Pair.Key.MsgIdPlural, Pair.Key.MsgCtxt);
 				if (FoundEntry.IsValid())
 				{
-					FoundEntry->AddExtractedComments(Pair.Value);
+					FoundEntry->AddExtractedComments(Pair.Value.ExtractedComments);
+					FoundEntry->AddTranslatorComments(Pair.Value.TranslatorComments);
 				}
 			}
 		}
diff --git a/Engine/Source/Developer/Localization/Public/PortableObjectFormatDOM.h b/Engine/Source/Developer/Localization/Public/PortableObjectFormatDOM.h
--- a/Engine/Source/Developer/Localization/Public/PortableObjectFormatDOM.h
+++ b/Engine/Source/Developer/Localization/Public/PortableObjectFormatDOM.h
@@ -142,28 +142,42 @@ public:
 	/**
 	 * Helper function that adds to the extracted comments.
 	 *
-	 * @param InReference	String representing extracted comment entries of the following form: "File/path/file.cpp:20 File/path/file.cpp:21 File/path/file2.cpp:5".
+	 * @param InComment	String representing an extracted comment entry.
 	 */
 	void AddExtractedComment( const FString& InComment );
 
 	/**
 	 * Helper function that adds to the reference comments.
 	 *
-	 * @param InReference	String representing reference comment entries.
+	 * @param InReference	String representing reference comment entries of the following form: "File/path/file.cpp:20 File/path/file.cpp:21 File/path/file2.cpp:5".
 	 */
 	void AddReference( const FString& InReference );
 
 	/**
-	 * Helper function that adds to the reference comments.
+	 * Helper function that adds to the extracted comments.
 	 *
-	 * @param InReferences	String array representing reference comment entries of the following form: "File/path/file.cpp:20 File/path/file.cpp:21 File/path/file2.cpp:5".
+	 * @param InComments	String array representing extracted comment entries.
 	 */
 	void AddExtractedComments( const TArray<FString>& InComments );
 
 	/**
-	 * Helper function that adds to the extracted comments.
+	 * Helper function that adds to the translator comments.
+	 *
+	 * @param InComment	String representing a translator comment entry.
+	 */
+	void AddTranslatorComment(const FString& InComment);
+
+	/**
+	 * Helper function that adds to the translator comments.
 	 *
-	 * @param InReferences	String array representing extracted comment entries
+	 * @param InComments	String array representing the translator's notes.
+	 */
+	void AddTranslatorComments(const TArray<FString>& InComments);
+
+	/**
+	 * Helper function that adds to the reference comments.
+	 *
+	 * @param InReferences	String array representing reference comment entries of the following form: "File/path/file.cpp:20 File/path/file.cpp:21 File/path/file2.cpp:5".
 	 */
 	void AddReferences( const TArray<FString>& InReferences );
 
