Fix Clang Win64 build.

#rnx

diff --git a/Engine/Source/Developer/Localization/Public/LocTextHelper.h b/Engine/Source/Developer/Localization/Public/LocTextHelper.h
--- a/Engine/Source/Developer/Localization/Public/LocTextHelper.h
+++ b/Engine/Source/Developer/Localization/Public/LocTextHelper.h
@@ -129,6 +129,10 @@ public:
 	 */
 	FString GetConflictReport() const;
 
+	/** Publicly movable */
+	FLocTextConflicts(FLocTextConflicts&&) = default;
+	FLocTextConflicts& operator=(FLocTextConflicts&&) = default;
+
 private:
 	FLocTextConflicts(const FLocTextConflicts&) = delete;
 	FLocTextConflicts& operator=(const FLocTextConflicts&) = delete;
