Updating copyright for Engine Developer.

#rnx

diff --git a/Engine/Source/Developer/Localization/Private/PortableObjectFormatDOM.cpp b/Engine/Source/Developer/Localization/Private/PortableObjectFormatDOM.cpp
--- a/Engine/Source/Developer/Localization/Private/PortableObjectFormatDOM.cpp
+++ b/Engine/Source/Developer/Localization/Private/PortableObjectFormatDOM.cpp
@@ -1,4 +1,4 @@
-// Copyright 1998-2019 Epic Games, Inc. All Rights Reserved.
+// Copyright Epic Games, Inc. All Rights Reserved.
 
 #include "PortableObjectFormatDOM.h"
 #include "Internationalization/Culture.h"
@@ -761,7 +761,7 @@ void FPortableObjectFormatDOM::CreateNewHeader()
 	//Hard code some header entries for now in the following format
 	/*
 	# Engine English translation
-	# Copyright 1998-2019 Epic Games, Inc. All Rights Reserved.
+	# Copyright Epic Games, Inc. All Rights Reserved.
 	#
 	msgid ""
 	msgstr ""
@@ -790,7 +790,7 @@ void FPortableObjectFormatDOM::CreateNewHeader()
 	Header.SetEntryValue( TEXT("Plural-Forms"), Language.GetPluralForms() );
 
 	Header.Comments.Add( FString::Printf(TEXT("%s %s translation."), *GetProjectName(), *Language.EnglishName() ) );
-	Header.Comments.Add( TEXT("Copyright 1998-2019 Epic Games, Inc. All Rights Reserved.") );
+	Header.Comments.Add( TEXT("Copyright Epic Games, Inc. All Rights Reserved.") );
 	Header.Comments.Add( FString(TEXT("")) );
 }
 
