Follow-up to CL# 11118619 to fix the class filtering on the UI option

#rnx

diff --git a/Engine/Source/Developer/Localization/Public/LocalizationTargetTypes.h b/Engine/Source/Developer/Localization/Public/LocalizationTargetTypes.h
--- a/Engine/Source/Developer/Localization/Public/LocalizationTargetTypes.h
+++ b/Engine/Source/Developer/Localization/Public/LocalizationTargetTypes.h
@@ -179,7 +179,7 @@ struct FGatherTextFromPackagesConfiguration
 	TArray<FName> Collections;
 
 	/* Classes that should be excluded from gathering. */
-	UPROPERTY(config, EditAnywhere, Category = "Filter")
+	UPROPERTY(config, EditAnywhere, Category = "Filter", meta=(MetaClass="Object", AllowAbstract="True"))
 	TArray<FSoftClassPath> ExcludeClasses;
 
 	/* Should classes derived from those in the exclude classes list also be excluded from gathering? */
