Fixed LatAm Spanish using Castilian Spanish translations for things that had been deliberately left untranslated

LatAm Spanish (es-419) often leaves words in English rather than translate them, however we were finding that they would take the Castilian Spanish (es) translation instead, which often lead to confusion for players as those terms weren't understood within the LatAm community.

This was caused by the fact that any untranslated text (text matching the source) was considered redundant and was skipped when generating the LocRes file, which failed to take into account potentially needing to override a base translation when dealing with language hierarchies.

This change now tests otherwise untranslated text against the translations of any parent languages, and will consider the entry translated if it is different from one that would be inherited at runtime.

diff --git a/Engine/Source/Developer/Localization/Private/TextLocalizationResourceGenerator.cpp b/Engine/Source/Developer/Localization/Private/TextLocalizationResourceGenerator.cpp
--- a/Engine/Source/Developer/Localization/Private/TextLocalizationResourceGenerator.cpp
+++ b/Engine/Source/Developer/Localization/Private/TextLocalizationResourceGenerator.cpp
@@ -23,6 +23,9 @@ bool FTextLocalizationResourceGenerator::GenerateLocRes(const FLocTextHelper& In
 	const bool bIsNativeCulture = InCultureToGenerate == InLocTextHelper.GetNativeCulture();
 	FCulturePtr Culture = FInternationalization::Get().GetCulture(InCultureToGenerate);
 
+	TArray<FString> InheritedCultures = Culture->GetPrioritizedParentCultureNames();
+	InheritedCultures.Remove(Culture->GetName());
+
 	// Always add the split platforms so that they generate an empty LocRes if there are no entries for that platform in the master manifest
 	for (const FString& SplitPlatformName : InLocTextHelper.GetPlatformsToSplit())
 	{
@@ -34,7 +37,7 @@ bool FTextLocalizationResourceGenerator::GenerateLocRes(const FLocTextHelper& In
 	}
 
 	// Add each manifest entry to the LocRes file
-	InLocTextHelper.EnumerateSourceTexts([&InLocTextHelper, &InCultureToGenerate, InGenerateFlags, &InLocResID, &OutPlatformAgnosticLocRes, &OutPerPlatformLocRes, InPriority, bIsNativeCulture, Culture](TSharedRef<FManifestEntry> InManifestEntry) -> bool
+	InLocTextHelper.EnumerateSourceTexts([&InLocTextHelper, &InCultureToGenerate, InGenerateFlags, &InLocResID, &OutPlatformAgnosticLocRes, &OutPerPlatformLocRes, InPriority, bIsNativeCulture, Culture, &InheritedCultures](TSharedRef<FManifestEntry> InManifestEntry) -> bool
 	{
 		// For each context, we may need to create a different or even multiple LocRes entries.
 		for (const FManifestContext& Context : InManifestEntry->Contexts)
@@ -44,7 +47,26 @@ bool FTextLocalizationResourceGenerator::GenerateLocRes(const FLocTextHelper& In
 			InLocTextHelper.GetRuntimeText(InCultureToGenerate, InManifestEntry->Namespace, Context.Key, Context.KeyMetadataObj, ELocTextExportSourceMethod::NativeText, InManifestEntry->Source, TranslationText, EnumHasAnyFlags(InGenerateFlags, EGenerateLocResFlags::AllowStaleTranslations));
 
 			// Is this entry considered translated? Native entries are always translated
-			const bool bIsTranslated = bIsNativeCulture || !InManifestEntry->Source.IsExactMatch(TranslationText);
+			bool bIsTranslated = bIsNativeCulture || !InManifestEntry->Source.IsExactMatch(TranslationText);
+			if (!bIsTranslated && InheritedCultures.Num() > 0)
+			{
+				// If this entry has parent languages then we also need to test whether the current translation is different from any parent that we have translations for, 
+				// as it may be that the translation was explicitly changed back to being the native text for some reason (eg, es-419 needs something in English that es translates)
+				for (const FString& InheritedCulture : InheritedCultures)
+				{
+					if (InLocTextHelper.HasArchive(InheritedCulture))
+					{
+						FLocItem InheritedText;
+						InLocTextHelper.GetRuntimeText(InheritedCulture, InManifestEntry->Namespace, Context.Key, Context.KeyMetadataObj, ELocTextExportSourceMethod::NativeText, InManifestEntry->Source, InheritedText, EnumHasAnyFlags(InGenerateFlags, EGenerateLocResFlags::AllowStaleTranslations));
+						if (!InheritedText.IsExactMatch(TranslationText))
+						{
+							bIsTranslated = true;
+							break;
+						}
+					}
+				}
+			}
+
 			if (bIsTranslated)
 			{
 				// Validate translations that look like they could be format patterns
