Removed the ELocTextPlatformSplitMode::Restricted option

There are no longer any platforms marked as having "restricted" localization, as they all now live inside platform extensions. The original spirit of this option lives on as ELocTextPlatformSplitMode::Confidential, however none of our localization targets are using it and it has limited utility.

diff --git a/Engine/Config/Localization/Category.ini b/Engine/Config/Localization/Category.ini
--- a/Engine/Config/Localization/Category.ini
+++ b/Engine/Config/Localization/Category.ini
@@ -1,11 +1,10 @@
-﻿[CommonSettings]
+[CommonSettings]
 ManifestDependencies=Content/Localization/Editor/Editor.manifest
 SourcePath=Content/Localization/Category
 DestinationPath=Content/Localization/Category
 ManifestName=Category.manifest
 ArchiveName=Category.archive
 PortableObjectName=Category.po
-PlatformSplitMode=ELocTextPlatformSplitMode::Restricted
 NativeCulture=en
 CulturesToGenerate=es
 CulturesToGenerate=en
diff --git a/Engine/Config/Localization/Editor.ini b/Engine/Config/Localization/Editor.ini
--- a/Engine/Config/Localization/Editor.ini
+++ b/Engine/Config/Localization/Editor.ini
@@ -1,11 +1,10 @@
-﻿[CommonSettings]
+[CommonSettings]
 ManifestDependencies=Content/Localization/Engine/Engine.manifest
 SourcePath=Content/Localization/Editor
 DestinationPath=Content/Localization/Editor
 ManifestName=Editor.manifest
 ArchiveName=Editor.archive
 PortableObjectName=Editor.po
-PlatformSplitMode=ELocTextPlatformSplitMode::Restricted
 NativeCulture=en
 CulturesToGenerate=en
 CulturesToGenerate=es
diff --git a/Engine/Config/Localization/EditorTutorials.ini b/Engine/Config/Localization/EditorTutorials.ini
--- a/Engine/Config/Localization/EditorTutorials.ini
+++ b/Engine/Config/Localization/EditorTutorials.ini
@@ -1,4 +1,4 @@
-﻿[CommonSettings]
+[CommonSettings]
 ManifestDependencies=Content/Localization/Engine/Engine.manifest
 ManifestDependencies=Content/Localization/Editor/Editor.manifest
 ManifestDependencies=Content/Localization/Category/Category.manifest
@@ -10,7 +10,6 @@ DestinationPath=Content/Localization/EditorTutorials
 ManifestName=EditorTutorials.manifest
 ArchiveName=EditorTutorials.archive
 PortableObjectName=EditorTutorials.po
-PlatformSplitMode=ELocTextPlatformSplitMode::Restricted
 NativeCulture=en
 CulturesToGenerate=es
 CulturesToGenerate=en
diff --git a/Engine/Config/Localization/Engine.ini b/Engine/Config/Localization/Engine.ini
--- a/Engine/Config/Localization/Engine.ini
+++ b/Engine/Config/Localization/Engine.ini
@@ -1,10 +1,9 @@
-﻿[CommonSettings]
+[CommonSettings]
 SourcePath=Content/Localization/Engine
 DestinationPath=Content/Localization/Engine
 ManifestName=Engine.manifest
 ArchiveName=Engine.archive
 PortableObjectName=Engine.po
-PlatformSplitMode=ELocTextPlatformSplitMode::Restricted
 NativeCulture=en
 CulturesToGenerate=en
 CulturesToGenerate=es
diff --git a/Engine/Config/Localization/Keywords.ini b/Engine/Config/Localization/Keywords.ini
--- a/Engine/Config/Localization/Keywords.ini
+++ b/Engine/Config/Localization/Keywords.ini
@@ -1,10 +1,9 @@
-﻿[CommonSettings]
+[CommonSettings]
 SourcePath=./Content/Localization/Keywords
 DestinationPath=./Content/Localization/Keywords
 ManifestName=Keywords.manifest
 ArchiveName=Keywords.archive
 PortableObjectName=Keywords.po
-PlatformSplitMode=ELocTextPlatformSplitMode::Restricted
 NativeCulture=en
 CulturesToGenerate=es
 CulturesToGenerate=en
diff --git a/Engine/Config/Localization/PropertyNames.ini b/Engine/Config/Localization/PropertyNames.ini
--- a/Engine/Config/Localization/PropertyNames.ini
+++ b/Engine/Config/Localization/PropertyNames.ini
@@ -1,11 +1,10 @@
-﻿[CommonSettings]
+[CommonSettings]
 ManifestDependencies=Content/Localization/Editor/Editor.manifest
 SourcePath=Content/Localization/PropertyNames
 DestinationPath=Content/Localization/PropertyNames
 ManifestName=PropertyNames.manifest
 ArchiveName=PropertyNames.archive
 PortableObjectName=PropertyNames.po
-PlatformSplitMode=ELocTextPlatformSplitMode::Restricted
 NativeCulture=en
 CulturesToGenerate=es
 CulturesToGenerate=en
diff --git a/Engine/Config/Localization/ToolTips.ini b/Engine/Config/Localization/ToolTips.ini
--- a/Engine/Config/Localization/ToolTips.ini
+++ b/Engine/Config/Localization/ToolTips.ini
@@ -1,11 +1,10 @@
-﻿[CommonSettings]
+[CommonSettings]
 ManifestDependencies=Content/Localization/Engine/Engine.manifest
 SourcePath=Content/Localization/ToolTips
 DestinationPath=Content/Localization/ToolTips
 ManifestName=ToolTips.manifest
 ArchiveName=ToolTips.archive
 PortableObjectName=ToolTips.po
-PlatformSplitMode=ELocTextPlatformSplitMode::Restricted
 NativeCulture=en
 CulturesToGenerate=es
 CulturesToGenerate=en
diff --git a/Engine/Source/Developer/Localization/Private/LocTextHelper.cpp b/Engine/Source/Developer/Localization/Private/LocTextHelper.cpp
--- a/Engine/Source/Developer/Localization/Private/LocTextHelper.cpp
+++ b/Engine/Source/Developer/Localization/Private/LocTextHelper.cpp
@@ -27,30 +27,21 @@ const TArray<FString>& FLocTextPlatformSplitUtils::GetPlatformsToSplit(const ELo
 {
 	switch (InSplitMode)
 	{
-	case ELocTextPlatformSplitMode::Restricted:
+	case ELocTextPlatformSplitMode::Confidential:
 		{
-			static TArray<FString> RestrictedPlatformNames = []()
+			static const TArray<FString> ConfidentialPlatformNames = []()
 			{
-				TArray<FString> TmpArray;
-				for (const FString& PlatformName : FDataDrivenPlatformInfoRegistry::GetConfidentialPlatforms())
-				{
-					const FDataDrivenPlatformInfoRegistry::FPlatformInfo& PlatformInfo = FDataDrivenPlatformInfoRegistry::GetPlatformInfo(PlatformName);
-					check(PlatformInfo.bIsConfidential);
-					if (PlatformInfo.bRestrictLocalization)
-					{
-						TmpArray.Add(PlatformName);
-					}
-				}
+				TArray<FString> TmpArray = FDataDrivenPlatformInfoRegistry::GetConfidentialPlatforms();
 				TmpArray.Sort();
 				return TmpArray;
 			}();
-			return RestrictedPlatformNames;
+			return ConfidentialPlatformNames;
 		}
 		break;
 
 	case ELocTextPlatformSplitMode::All:
 		{
-			static TArray<FString> AllPlatformNames = []()
+			static const TArray<FString> AllPlatformNames = []()
 			{
 				TArray<FString> TmpArray;
 				for (const PlatformInfo::FPlatformInfo& Info : PlatformInfo::GetPlatformInfoArray())
diff --git a/Engine/Source/Developer/Localization/Public/LocTextHelper.h b/Engine/Source/Developer/Localization/Public/LocTextHelper.h
--- a/Engine/Source/Developer/Localization/Public/LocTextHelper.h
+++ b/Engine/Source/Developer/Localization/Public/LocTextHelper.h
@@ -38,8 +38,8 @@ enum class ELocTextPlatformSplitMode : uint8
 {
 	/** Don't split platform specific localization data */
 	None,
-	/** Split platform specific localization data for restricted platforms only */
-	Restricted,
+	/** Split platform specific localization data for confidential platforms only */
+	Confidential,
 	/** Split platform specific localization data for all platforms */
 	All,
 };
diff --git a/Engine/Source/Runtime/Core/Private/Misc/DataDrivenPlatformInfoRegistry.cpp b/Engine/Source/Runtime/Core/Private/Misc/DataDrivenPlatformInfoRegistry.cpp
--- a/Engine/Source/Runtime/Core/Private/Misc/DataDrivenPlatformInfoRegistry.cpp
+++ b/Engine/Source/Runtime/Core/Private/Misc/DataDrivenPlatformInfoRegistry.cpp
@@ -157,7 +157,6 @@ static void DDPIGetStringArray(const FConfigFile& IniFile, const TCHAR* Key, TAr
 static void LoadDDPIIniSettings(const FConfigFile& IniFile, FDataDrivenPlatformInfoRegistry::FPlatformInfo& Info)
 {
 	DDPIGetBool(IniFile, TEXT("bIsConfidential"), Info.bIsConfidential);
-	DDPIGetBool(IniFile, TEXT("bRestrictLocalization"), Info.bRestrictLocalization);
 	DDPIGetString(IniFile, TEXT("AudioCompressionSettingsIniSectionName"), Info.AudioCompressionSettingsIniSectionName);
 	DDPIGetStringArray(IniFile, TEXT("AdditionalRestrictedFolders"), Info.AdditionalRestrictedFolders);
 	
diff --git a/Engine/Source/Runtime/Core/Public/Misc/DataDrivenPlatformInfoRegistry.h b/Engine/Source/Runtime/Core/Public/Misc/DataDrivenPlatformInfoRegistry.h
--- a/Engine/Source/Runtime/Core/Public/Misc/DataDrivenPlatformInfoRegistry.h
+++ b/Engine/Source/Runtime/Core/Public/Misc/DataDrivenPlatformInfoRegistry.h
@@ -20,9 +20,6 @@ struct CORE_API FDataDrivenPlatformInfoRegistry
 		// is this platform confidential
 		bool bIsConfidential = false;
 
-		// should this platform be split when using ELocTextPlatformSplitMode::Restricted (only used when bIsConfidential is true)
-		bool bRestrictLocalization = false;
-
 		// the name of the ini section to use to load audio compression settings (used at runtime and cooktime)
 		FString AudioCompressionSettingsIniSectionName;
 
