Lyra/Localization - Adding some comments to the generated localization files so that people don't accidentally hand edit them thinking they are the source.

diff --git a/Engine/Source/Developer/Localization/Private/LocalizationConfigurationScript.cpp b/Engine/Source/Developer/Localization/Private/LocalizationConfigurationScript.cpp
--- a/Engine/Source/Developer/Localization/Private/LocalizationConfigurationScript.cpp
+++ b/Engine/Source/Developer/Localization/Private/LocalizationConfigurationScript.cpp
@@ -1199,6 +1199,10 @@ namespace LocalizationConfigurationScript
 			{
 				return false;
 			}
+			else
+			{
+				NewConfigContents = TEXT("; THESE ARE GENERATED FILES, DO NOT EDIT DIRECTLY!\r\n; USE THE LOCALIZATION DASHBOARD IN THE UNREAL EDITOR TO EDIT THE CONFIGURATION\r\n") + NewConfigContents;
+			}
 		}
 
 		bool bWriteFile = true;
