Localization conflict reports can now be output as .csv files.
- Introduced EConflictReportFormat to determine the output format of the conflicts report.
- Updated the localization .ini files to specify the requested output format.

diff --git a/Engine/Config/Localization/Category.ini b/Engine/Config/Localization/Category.ini
--- a/Engine/Config/Localization/Category.ini
+++ b/Engine/Config/Localization/Category.ini
@@ -61,4 +61,6 @@ DestinationPath=\\epicgames.net\root\UE3\Localization\WordCounts
 bWordCountReport=true
 WordCountReportName=Category.csv
 bConflictReport=true
-ConflictReportName=CategoryConflicts.txt
\ No newline at end of file
+ConflictReportFormat=EConflictReportFormat::CSV
+; Do not add an extension to the filename. EConflictReportFormat will be used to append the correct extension
+ConflictReportName=CategoryConflicts
\ No newline at end of file
diff --git a/Engine/Config/Localization/Editor.ini b/Engine/Config/Localization/Editor.ini
--- a/Engine/Config/Localization/Editor.ini
+++ b/Engine/Config/Localization/Editor.ini
@@ -106,4 +106,6 @@ DestinationPath=\\epicgames.net\root\UE3\Localization\WordCounts
 bWordCountReport=true
 WordCountReportName=Editor.csv
 bConflictReport=true
-ConflictReportName=EditorConflicts.txt
\ No newline at end of file
+ConflictReportFormat=EConflictReportFormat::CSV
+; Do not add an extension to the filename. EConflictReportFormat will be used to append the correct extension
+ConflictReportName=EditorConflicts
\ No newline at end of file
diff --git a/Engine/Config/Localization/EditorTutorials.ini b/Engine/Config/Localization/EditorTutorials.ini
--- a/Engine/Config/Localization/EditorTutorials.ini
+++ b/Engine/Config/Localization/EditorTutorials.ini
@@ -62,4 +62,6 @@ DestinationPath=\\epicgames.net\root\UE3\Localization\WordCounts
 bWordCountReport=true
 WordCountReportName=EditorTutorials.csv
 bConflictReport=true
-ConflictReportName=EditorTutorialsConflicts.txt
\ No newline at end of file
+ConflictReportFormat=EConflictReportFormat::CSV
+; Do not add an extension to the filename. EConflictReportFormat will be used to append the correct extension
+ConflictReportName=EditorTutorialsConflicts
\ No newline at end of file
diff --git a/Engine/Config/Localization/Engine.ini b/Engine/Config/Localization/Engine.ini
--- a/Engine/Config/Localization/Engine.ini
+++ b/Engine/Config/Localization/Engine.ini
@@ -84,4 +84,6 @@ DestinationPath=\\epicgames.net\root\UE3\Localization\WordCounts
 bWordCountReport=true
 WordCountReportName=Engine.csv
 bConflictReport=true
-ConflictReportName=EngineConflicts.txt
\ No newline at end of file
+ConflictReportFormat=EConflictReportFormat::CSV
+; Do not add an extension to the filename. EConflictReportFormat will be used to append the correct extension
+ConflictReportName=EngineConflicts
\ No newline at end of file
diff --git a/Engine/Config/Localization/Keywords.ini b/Engine/Config/Localization/Keywords.ini
--- a/Engine/Config/Localization/Keywords.ini
+++ b/Engine/Config/Localization/Keywords.ini
@@ -57,4 +57,6 @@ DestinationPath=\\epicgames.net\root\UE3\Localization\WordCounts
 bWordCountReport=true
 WordCountReportName=Keywords.csv
 bConflictReport=true
-ConflictReportName=KeywordsConflicts.txt
\ No newline at end of file
+ConflictReportFormat=EConflictReportFormat::CSV
+; Do not add an extension to the filename. EConflictReportFormat will be used to append the correct extension
+ConflictReportName=KeywordsConflicts
\ No newline at end of file
diff --git a/Engine/Config/Localization/PropertyNames.ini b/Engine/Config/Localization/PropertyNames.ini
--- a/Engine/Config/Localization/PropertyNames.ini
+++ b/Engine/Config/Localization/PropertyNames.ini
@@ -61,4 +61,6 @@ DestinationPath=\\epicgames.net\root\UE3\Localization\WordCounts
 bWordCountReport=true
 WordCountReportName=PropertyNames.csv
 bConflictReport=true
-ConflictReportName=PropertyNamesConflicts.txt
\ No newline at end of file
+ConflictReportFormat=EConflictReportFormat::CSV
+; Do not add an extension to the filename. EConflictReportFormat will be used to append the correct extension
+ConflictReportName=PropertyNamesConflicts
\ No newline at end of file
diff --git a/Engine/Config/Localization/ToolTips.ini b/Engine/Config/Localization/ToolTips.ini
--- a/Engine/Config/Localization/ToolTips.ini
+++ b/Engine/Config/Localization/ToolTips.ini
@@ -64,4 +64,6 @@ DestinationPath=\\epicgames.net\root\UE3\Localization\WordCounts
 bWordCountReport=true
 WordCountReportName=ToolTips.csv
 bConflictReport=true
-ConflictReportName=ToolTipsConflicts.txt
\ No newline at end of file
+ConflictReportFormat=EConflictReportFormat::CSV
+; Do not add an extension to the filename. EConflictReportFormat will be used to append the correct extension
+ConflictReportName=ToolTipsConflicts
\ No newline at end of file
diff --git a/Engine/Source/Developer/Localization/Private/LocTextHelper.cpp b/Engine/Source/Developer/Localization/Private/LocTextHelper.cpp
--- a/Engine/Source/Developer/Localization/Private/LocTextHelper.cpp
+++ b/Engine/Source/Developer/Localization/Private/LocTextHelper.cpp
@@ -114,12 +114,15 @@ FString FLocTextConflicts::GetConflictReport() const
 		const FString& Key = Conflict->Key.GetString();
 
 		bool bAddToReport = false;
+		// Grab the list of all conflicting source text associated with this key 
 		TArray<FLocItem> SourceList;
 		Conflict->EntriesBySourceLocation.GenerateValueArray(SourceList);
+
 		if (SourceList.Num() >= 2)
 		{
 			for (int32 i = 0; i < SourceList.Num() - 1 && !bAddToReport; ++i)
 			{
+				// we only add the conflict to the report if the source strings do not match up.
 				for (int32 j = i + 1; j < SourceList.Num() && !bAddToReport; ++j)
 				{
 					if (!(SourceList[i] == SourceList[j]))
@@ -154,6 +157,87 @@ FString FLocTextConflicts::GetConflictReport() const
 	return Report;
 }
 
+FString FLocTextConflicts::GetConflictReportAsCSV() const
+{
+	// First have the headers 
+	FString Report = TEXT("Namespace,Key,File1,Line1,File2,Line2,Source1,Source2\n");
+
+	for (const auto& ConflictPair : EntriesByKey)
+	{
+		const TSharedRef<FConflict>& Conflict = ConflictPair.Value;
+		const FString& Namespace = Conflict->Namespace.GetString();
+		const FString& Key = Conflict->Key.GetString();
+
+		bool bAddToReport = false;
+		TArray<FLocItem> SourceList;
+		Conflict->EntriesBySourceLocation.GenerateValueArray(SourceList);
+		if (SourceList.Num() >= 2)
+		{
+			for (int32 i = 0; i < SourceList.Num() - 1 && !bAddToReport; ++i)
+			{
+				for (int32 j = i + 1; j < SourceList.Num() && !bAddToReport; ++j)
+				{
+					// Conflict only if the source text and metadata don't match 
+					if (!(SourceList[i] == SourceList[j]))
+					{
+						bAddToReport = true;
+					}
+				}
+			}
+		}
+
+		if (bAddToReport)
+		{
+			// the format of the CSV will need file1,line1 and file2,line2 to be compared against each other 
+			// If a single key has more than 1 conflict, we want file 1 in the csv to always be the same one.
+			// Thus we store the results of the first entry to be appended in each row of the CSV for a single entry 
+
+			bool bFirstEntry = true;
+			FString FirstEntrySourceFile;
+			FString FirstEntryLineNumber;
+			FString FirstEntrySourceText;
+
+			// This maps each of the source locations to its conflicting source string 
+			// for each entry in this, we create a new row in the .csv 
+			for (auto EntryIter = Conflict->EntriesBySourceLocation.CreateConstIterator(); EntryIter; ++EntryIter)
+			{
+				const FString& SourceLocation = EntryIter.Key();
+				FString ProcessedSourceLocation = SourceLocation;
+				ProcessedSourceLocation.ReplaceInline(TEXT("\\"), TEXT("/"));
+
+				const FString& SourceText = EntryIter.Value().Text.ReplaceCharWithEscapedChar();
+				// in the event there are more than 2 key conflicts for a single key,
+				// this will always be the entry for file 1 and line 1 in the output file 
+				if (bFirstEntry)
+				{
+					// For a source file conflict, source location is in the format My/File/Name.cpp:LineNumber
+					bool bSplitSuccessful = ProcessedSourceLocation.Split(TEXT(":"), &FirstEntrySourceFile, &FirstEntryLineNumber, ESearchCase::IgnoreCase, ESearchDir::FromEnd);
+					// If the split fails, this means that the location is not in the form file:linenumber.
+					// if the line number string is not numeric, this means we are dealing with an asset path 
+					if (!bSplitSuccessful  || !FirstEntryLineNumber.IsNumeric())
+					{
+						FirstEntrySourceFile = ProcessedSourceLocation;
+						FirstEntryLineNumber = TEXT("0");
+					}
+					FirstEntrySourceText = SourceText;
+					bFirstEntry = false;
+					continue;
+				}
+				FString CurrentSourceFile;
+				FString CurrentLineNumber;
+				bool bSplitSuccessful = ProcessedSourceLocation.Split(TEXT(":"), &CurrentSourceFile, &CurrentLineNumber, ESearchCase::IgnoreCase, ESearchDir::FromEnd);
+				// Same as above 
+				if (!bSplitSuccessful || !CurrentLineNumber.IsNumeric())
+				{
+					CurrentSourceFile = ProcessedSourceLocation;
+					CurrentLineNumber = TEXT("0");
+				}
+				Report += FString::Printf(TEXT("%s,%s,%s,%s,%s,%s,%s,%s\n"), *Namespace, *Key, *FirstEntrySourceFile, *FirstEntryLineNumber, *CurrentSourceFile, *CurrentLineNumber, *FirstEntrySourceText, *SourceText);
+			}
+		}
+	}
+	return Report;
+}
 
 const FString FLocTextWordCounts::ColHeadingDateTime = TEXT("Date/Time");
 const FString FLocTextWordCounts::ColHeadingWordCount = TEXT("Word Count");
@@ -1104,7 +1188,7 @@ FString FLocTextHelper::GetConflictReport() const
 	return ConflictTracker.GetConflictReport();
 }
 
-bool FLocTextHelper::SaveConflictReport(const FString& InReportFilePath, FText* OutError) const
+bool FLocTextHelper::SaveConflictReport(const FString& InReportFilePath, EConflictReportFormat InConflictReportFormat, FText* OutError) const
 {
 	bool bSaved = false;
 
@@ -1112,8 +1196,29 @@ bool FLocTextHelper::SaveConflictReport(const FString& InReportFilePath, FText*
 	{
 		LocFileNotifies->PreFileWrite(InReportFilePath);
 	}
-
-	const FString ConflictReport = ConflictTracker.GetConflictReport();
+	FString ConflictReport;
+	switch (InConflictReportFormat)
+	{
+		case EConflictReportFormat::Txt: 
+		{
+			ConflictReport = ConflictTracker.GetConflictReport();
+			break;
+		}
+		case EConflictReportFormat::CSV:
+		{
+			ConflictReport = ConflictTracker.GetConflictReportAsCSV();
+			break;
+		}
+		default:
+		{
+			if (OutError)
+			{
+				*OutError = FText::Format(LOCTEXT("Error_UnsupportedConflictReportFormat", "Failed to generate conflict report '{0}'. An unsupported conflict report format was specified. Supported formats can be found in EConflictReportFormat."), FText::FromString(InReportFilePath));
+			}
+			return false;
+		}
+	}
+	
 	if (FFileHelper::SaveStringToFile(ConflictReport, *InReportFilePath))
 	{
 		bSaved = true;
diff --git a/Engine/Source/Developer/Localization/Public/LocTextHelper.h b/Engine/Source/Developer/Localization/Public/LocTextHelper.h
--- a/Engine/Source/Developer/Localization/Public/LocTextHelper.h
+++ b/Engine/Source/Developer/Localization/Public/LocTextHelper.h
@@ -78,6 +78,13 @@ struct LOCALIZATION_API FLocTextPlatformSplitUtils
 	/** Get the platforms names that should be split, based on the given split mode */
 	static const TArray<FString>& GetPlatformsToSplit(const ELocTextPlatformSplitMode& InSplitMode);
 };
+/** An enum representing all the formats the conflict report can be output in. */
+enum class EConflictReportFormat : uint8
+{
+	None,
+	Txt,
+	CSV
+};
 
 /**
  * Class that tracks any conflicts that occur when gathering source text entries.
@@ -123,11 +130,15 @@ public:
 	 * @param InSourceLocation		The source location of the conflict.
 	 */
 	void AddConflict(const FLocKey& InNamespace, const FLocKey& InKey, const TSharedPtr<FLocMetadataObject>& InKeyMetadata, const FLocItem& InSource, const FString& InSourceLocation);
-
+	
 	/**
-	 * Convert the conflicts to a string format that can be easily saved as a report summary.
+	 * Convert the conflicts to a string format that can be easily saved as a .txt report summary.
 	 */
 	FString GetConflictReport() const;
+	/**
+	 * Convert the conflicts to a string format that can be easily saved as a .csv report summary.
+	 */
+	FString GetConflictReportAsCSV() const;
 
 private:
 	FLocTextConflicts(const FLocTextConflicts&) = delete;
@@ -778,16 +789,18 @@ public:
 	 * Get a conflict report that can be easily saved as a report summary.
 	 */
 	FString GetConflictReport() const;
+	
 
 	/**
 	 * Save the conflict report summary to disk.
 	 *
 	 * @param InReportFilePath		Full file path to write the report to.
+	 * @param InConflictReportFormat		The format the conflict report should be generated in.
 	 * @param OutError				Optional text to be filled when an error occurs.
 	 *
 	 * @return True if the file was saved, false otherwise.
 	 */
-	bool SaveConflictReport(const FString& InReportFilePath, FText* OutError = nullptr) const;
+	bool SaveConflictReport(const FString& InReportFilePath, EConflictReportFormat InConflictReportFormat, FText* OutError = nullptr) const;
 
 	/**
 	 * Get a word count report for the current state of the manifest and archives.
diff --git a/Engine/Source/Editor/UnrealEd/Classes/Commandlets/GenerateTextLocalizationReportCommandlet.h b/Engine/Source/Editor/UnrealEd/Classes/Commandlets/GenerateTextLocalizationReportCommandlet.h
--- a/Engine/Source/Editor/UnrealEd/Classes/Commandlets/GenerateTextLocalizationReportCommandlet.h
+++ b/Engine/Source/Editor/UnrealEd/Classes/Commandlets/GenerateTextLocalizationReportCommandlet.h
@@ -7,7 +7,6 @@
 #include "Commandlets/GatherTextCommandletBase.h"
 #include "GenerateTextLocalizationReportCommandlet.generated.h"
 
-
 UCLASS()
 class UGenerateTextLocalizationReportCommandlet : public UGatherTextCommandletBase
 {
@@ -18,6 +17,10 @@ public:
 	virtual int32 Main(const FString& Params) override;
 	//~ End UCommandlet Interface
 
+	virtual bool ShouldRunInPreview(const TArray<FString>& Switches, const TMap<FString, FString>& ParamVals) const
+	{
+		return true;
+	}
 private:
 	/**
 	 * Collects word count info and writes out a report.
@@ -34,7 +37,7 @@ private:
 	 * @return True if successful; false oterwise.
 	 */
 	bool ProcessConflictReport(const FString& DestinationPath);
-
+	
 private:
 	FString GatherTextConfigPath;
 	FString SectionName;
diff --git a/Engine/Source/Editor/UnrealEd/Private/Commandlets/GenerateTextLocalizationReportCommandlet.cpp b/Engine/Source/Editor/UnrealEd/Private/Commandlets/GenerateTextLocalizationReportCommandlet.cpp
--- a/Engine/Source/Editor/UnrealEd/Private/Commandlets/GenerateTextLocalizationReportCommandlet.cpp
+++ b/Engine/Source/Editor/UnrealEd/Private/Commandlets/GenerateTextLocalizationReportCommandlet.cpp
@@ -4,6 +4,8 @@
 #include "Misc/FileHelper.h"
 #include "Misc/Paths.h"
 #include "Internationalization/LocKeyFuncs.h"
+#include "LocTextHelper.h"
+#include "Misc/ConfigCacheIni.h"
 
 DEFINE_LOG_CATEGORY_STATIC(LogGenerateTextLocalizationReportCommandlet, Log, All);
 
@@ -30,7 +32,6 @@ int32 UGenerateTextLocalizationReportCommandlet::Main(const FString& Params)
 	else
 	{
 		UE_LOG(LogGenerateTextLocalizationReportCommandlet, Error, TEXT("No config specified."));
-		return -1;
 	}
 
 	// Set config section.
@@ -91,7 +92,7 @@ int32 UGenerateTextLocalizationReportCommandlet::Main(const FString& Params)
 
 	if( bConflictReport )
 	{
-		if( !ProcessConflictReport( DestinationPath ) )
+		if( !ProcessConflictReport( DestinationPath) )
 		{
 			UE_LOG(LogGenerateTextLocalizationReportCommandlet, Error, TEXT("Failed to generate localization conflict report."));
 			return -1;
@@ -177,11 +178,64 @@ bool UGenerateTextLocalizationReportCommandlet::ProcessConflictReport(const FStr
 		UE_LOG(LogGenerateTextLocalizationReportCommandlet, Error, TEXT("No conflict report name specified."));
 		return false;
 	}
-
+	if (ConflictReportName.Contains(TEXT("."), ESearchCase::IgnoreCase, ESearchDir::FromEnd))
+	{
+		UE_LOG(LogGenerateTextLocalizationReportCommandlet, Error, TEXT("Conflict report name has a file extension. Please remove the file extension from the name. EConflictReportFormat is used to determine the file extension."));
+		return false;
+	}
+	EConflictReportFormat ConflictReportFormat = EConflictReportFormat::None;
+	FString ConflictReportFormatString;
+	if (!GConfig->GetString(*SectionName, TEXT("ConflictReportFormat"), ConflictReportFormatString, GatherTextConfigPath))
+	{
+		UE_LOG(LogGenerateTextLocalizationReportCommandlet, Display, TEXT("Conflict report format not specified. Conflict report will default to CSV format."));
+		ConflictReportFormat = EConflictReportFormat::CSV;
+	}
+	else
+	{
+		ConflictReportFormatString.TrimStartAndEndInline();
+		static const TCHAR* TxtEnumString = TEXT("EConflictReportFormat::Txt");
+		static const TCHAR* CSVEnumString = TEXT("EConflictReportFormat::CSV");
+		if (ConflictReportFormatString.Equals(TxtEnumString))
+		{
+			ConflictReportFormat = EConflictReportFormat::Txt;
+		}
+		else if (ConflictReportFormatString.Equals(CSVEnumString))
+		{
+			ConflictReportFormat = EConflictReportFormat::CSV;
+		}
+		else
+		{
+			// @TODOLocalization: Consider alternatives to defaulting to CSV.
+			UE_LOG(LogGenerateTextLocalizationReportCommandlet, Warning, TEXT("Specified conflict report format %s in \'%s\' not supported. Defaulting to CSV format."), *ConflictReportFormatString, *GatherTextConfigPath);
+			ConflictReportFormat = EConflictReportFormat::CSV;
+		}
+	}
+	switch (ConflictReportFormat)
+	{
+		case EConflictReportFormat::CSV:
+		{
+			UE_LOG(LogGenerateTextLocalizationReportCommandlet, Display, TEXT("Conflict report format will be in CSV."));
+			static const TCHAR* CSVExtension = TEXT(".csv");
+			ConflictReportName += CSVExtension;
+			break;
+		}
+		case EConflictReportFormat::Txt:
+		{
+			UE_LOG(LogGenerateTextLocalizationReportCommandlet, Display, TEXT("Conflict report formatwill be in in txt format."));
+			static const TCHAR* TxtExtension = TEXT(".txt");
+			ConflictReportName += TxtExtension;
+			break;
+		}
+		default:
+		{
+			UE_LOG(LogGenerateTextLocalizationReportCommandlet, Error, TEXT("Unsupported conflict report format %s detected in %s. Unable to create the appropriate report name and extension."), *ConflictReportFormatString, *GatherTextConfigPath);
+			return false;
+		}
+	}
 	const FString ReportFilePath = (DestinationPath / ConflictReportName);
 
 	FText ReportSaveError;
-	if (!GatherManifestHelper->SaveConflictReport(ReportFilePath, &ReportSaveError))
+	if (!GatherManifestHelper->SaveConflictReport(ReportFilePath, ConflictReportFormat , &ReportSaveError))
 	{
 		UE_LOG(LogGenerateTextLocalizationReportCommandlet, Error, TEXT("%s"), *ReportSaveError.ToString());
 		return false;
