Remove/replace uses of master/slave found in Localization with base folder Developer(only comments)

diff --git a/Engine/Source/Developer/Localization/Private/LocTextHelper.cpp b/Engine/Source/Developer/Localization/Private/LocTextHelper.cpp
--- a/Engine/Source/Developer/Localization/Private/LocTextHelper.cpp
+++ b/Engine/Source/Developer/Localization/Private/LocTextHelper.cpp
@@ -1677,7 +1677,7 @@ bool FLocTextHelper::SaveManifestImpl(const TSharedRef<const FInternationalizati
 		TSharedRef<FInternationalizationManifest> PlatformAgnosticManifest = MakeShared<FInternationalizationManifest>();
 		TMap<FName, TSharedRef<FInternationalizationManifest>> PerPlatformManifests;
 		{
-			// Always add the split platforms so that they generate an empty manifest if there are no entries for that platform in the master manifest
+			// Always add the split platforms so that they generate an empty manifest if there are no entries for that platform in the platform agnostic manifest
 			for (const FString& SplitPlatformName : GetPlatformsToSplit())
 			{
 				PerPlatformManifests.Add(*SplitPlatformName, MakeShared<FInternationalizationManifest>());
@@ -1862,7 +1862,7 @@ bool FLocTextHelper::SaveArchiveImpl(const TSharedRef<const FInternationalizatio
 		TSharedRef<FInternationalizationArchive> PlatformAgnosticArchive = MakeShared<FInternationalizationArchive>();
 		TMap<FName, TSharedRef<FInternationalizationArchive>> PerPlatformArchives;
 		{
-			// Always add the split platforms so that they generate an empty archives if there are no entries for that platform in the master archive
+			// Always add the split platforms so that they generate an empty archives if there are no entries for that platform in the platform agnostic archive
 			for (const FString& SplitPlatformName : GetPlatformsToSplit())
 			{
 				PerPlatformArchives.Add(*SplitPlatformName, MakeShared<FInternationalizationArchive>());
diff --git a/Engine/Source/Developer/Localization/Private/PortableObjectPipeline.cpp b/Engine/Source/Developer/Localization/Private/PortableObjectPipeline.cpp
--- a/Engine/Source/Developer/Localization/Private/PortableObjectPipeline.cpp
+++ b/Engine/Source/Developer/Localization/Private/PortableObjectPipeline.cpp
@@ -67,7 +67,7 @@ namespace
 
 	void BuildCollapsedManifest(FLocTextHelper& InLocTextHelper, const ELocalizedTextCollapseMode InTextCollapseMode, FCollapsedData& OutCollapsedData, TSharedPtr<FInternationalizationManifest>& OutPlatformAgnosticManifest, TMap<FName, TSharedRef<FInternationalizationManifest>>& OutPerPlatformManifests)
 	{
-		// Always add the split platforms so that they generate an empty manifest if there are no entries for that platform in the master manifest
+		// Always add the split platforms so that they generate an empty manifest if there are no entries for that platform in the platform agnostic manifest
 		OutPlatformAgnosticManifest = MakeShared<FInternationalizationManifest>();
 		for (const FString& SplitPlatformName : InLocTextHelper.GetPlatformsToSplit())
 		{
diff --git a/Engine/Source/Developer/Localization/Private/TextLocalizationResourceGenerator.cpp b/Engine/Source/Developer/Localization/Private/TextLocalizationResourceGenerator.cpp
--- a/Engine/Source/Developer/Localization/Private/TextLocalizationResourceGenerator.cpp
+++ b/Engine/Source/Developer/Localization/Private/TextLocalizationResourceGenerator.cpp
@@ -28,7 +28,7 @@ bool FTextLocalizationResourceGenerator::GenerateLocRes(const FLocTextHelper& In
 	TArray<FString> InheritedCultures = Culture->GetPrioritizedParentCultureNames();
 	InheritedCultures.Remove(Culture->GetName());
 
-	// Always add the split platforms so that they generate an empty LocRes if there are no entries for that platform in the master manifest
+	// Always add the split platforms so that they generate an empty LocRes if there are no entries for that platform in the platform agnostic manifest
 	for (const FString& SplitPlatformName : InLocTextHelper.GetPlatformsToSplit())
 	{
 		const FName SplitPlatformFName = *SplitPlatformName;
