Moved strings and localization to the moderation manifest.
Fixed commandline options for moderation generation commandlet.

diff --git a/Engine/Source/Developer/Localization/Private/LocTextHelper.cpp b/Engine/Source/Developer/Localization/Private/LocTextHelper.cpp
--- a/Engine/Source/Developer/Localization/Private/LocTextHelper.cpp
+++ b/Engine/Source/Developer/Localization/Private/LocTextHelper.cpp
@@ -587,6 +587,11 @@ bool FLocTextHelper::SaveManifest(const FString& InManifestFilePath, FText* OutE
 	return SaveManifestImpl(Manifest.ToSharedRef(), InManifestFilePath, OutError);
 }
 
+bool FLocTextHelper::SerializeManifestToJson(TSharedRef<FJsonObject> JsonObject)
+{
+	return FJsonInternationalizationManifestSerializer::SerializeManifest(Manifest.ToSharedRef(), JsonObject);
+}
+
 void FLocTextHelper::TrimManifest()
 {
 	if (Dependencies.Num() > 0)
diff --git a/Engine/Source/Developer/Localization/Public/LocTextHelper.h b/Engine/Source/Developer/Localization/Public/LocTextHelper.h
--- a/Engine/Source/Developer/Localization/Public/LocTextHelper.h
+++ b/Engine/Source/Developer/Localization/Public/LocTextHelper.h
@@ -369,6 +369,13 @@ public:
 	 */
 	bool SaveManifest(FText* OutError = nullptr) const;
 
+
+	/**
+	 * Attempt to serialize the entire manifest to an input json object
+	 * @return True if the serialization succeeds false otherwise.
+	 */
+	bool SerializeManifestToJson(TSharedRef<FJsonObject> JsonObject);
+
 	/**
 	 * Attempt to save the manifest file to the given file path.
 	 *
