- Compile issue fix.

diff --git a/Engine/Source/Developer/Localization/Public/LocTextHelper.h b/Engine/Source/Developer/Localization/Public/LocTextHelper.h
--- a/Engine/Source/Developer/Localization/Public/LocTextHelper.h
+++ b/Engine/Source/Developer/Localization/Public/LocTextHelper.h
@@ -24,6 +24,7 @@ class FArchiveEntry;
 class FInternationalizationArchive;
 class FLocMetadataObject;
 class FText;
+class FJsonObject;
 
 /** Flags controlling the behavior used when loading manifests and archives into FLocTextHelper */
 enum class ELocTextHelperLoadFlags : uint8
