PR #9270: Bugfix: remove unnecessary quotes of OutputKeys in FooTarget_Gather.ini when gather texts from meta data via Localization Dashboard. (Contributed by [REDACTED_PERSON])

#9270

diff --git a/Engine/Source/Developer/Localization/Private/LocalizationConfigurationScript.cpp b/Engine/Source/Developer/Localization/Private/LocalizationConfigurationScript.cpp
--- a/Engine/Source/Developer/Localization/Private/LocalizationConfigurationScript.cpp
+++ b/Engine/Source/Developer/Localization/Private/LocalizationConfigurationScript.cpp
@@ -424,7 +424,7 @@ namespace LocalizationConfigurationScript
 			{
 				ConfigSection.Add( TEXT("InputKeys"), Specification.MetaDataKey.Name );
 				ConfigSection.Add( TEXT("OutputNamespaces"), Specification.TextNamespace );
-				ConfigSection.Add( TEXT("OutputKeys"), FString::Printf(TEXT("\"%s\""), *Specification.TextKeyPattern.Pattern) );
+				ConfigSection.Add(TEXT("OutputKeys"), Specification.TextKeyPattern.Pattern);
 			}
 
 			ConfigSection.Add( TEXT("ShouldGatherFromEditorOnlyData"), Target->Settings.GatherFromMetaData.ShouldGatherFromEditorOnlyData ? TEXT("true") : TEXT("false") );
