Removed redundant Engine dependency from the Localization module

#rnx

diff --git a/Engine/Source/Developer/Localization/Localization.Build.cs b/Engine/Source/Developer/Localization/Localization.Build.cs
--- a/Engine/Source/Developer/Localization/Localization.Build.cs
+++ b/Engine/Source/Developer/Localization/Localization.Build.cs
@@ -21,7 +21,6 @@ public class Localization : ModuleRules
 				"Slate",
 				"SlateCore",
 				"InputCore",
-				"Engine",
 				"SourceControl",
 				"Json",
 			}
diff --git a/Engine/Source/Developer/Localization/Public/LocalizationTargetTypes.h b/Engine/Source/Developer/Localization/Public/LocalizationTargetTypes.h
--- a/Engine/Source/Developer/Localization/Public/LocalizationTargetTypes.h
+++ b/Engine/Source/Developer/Localization/Public/LocalizationTargetTypes.h
@@ -6,7 +6,6 @@
 #include "Containers/UnrealString.h"
 #include "CoreMinimal.h"
 #include "CoreTypes.h"
-#include "Engine/EngineTypes.h"
 #include "HAL/PlatformCrt.h"
 #include "Internationalization/Text.h"
 #include "Misc/AssertionMacros.h"
