Added base support for UGC localization

UGC localization can be used to provide a simplified localization experience for basic plugins (only providing the PO files to be translated), and is primarily designed for DLC plugins where the UGC localization will be compiled during cook (see UUserGeneratedContentLocalizationSettings).

Support for non-DLC plugins can be provided via project specific tooling built upon this base API.

Support for complex plugins (such as those containing different kinds of modules, eg) a mix of game/engine and editor) are not supported via this API.

diff --git a/Engine/Source/Developer/Localization/Localization.Build.cs b/Engine/Source/Developer/Localization/Localization.Build.cs
--- a/Engine/Source/Developer/Localization/Localization.Build.cs
+++ b/Engine/Source/Developer/Localization/Localization.Build.cs
@@ -23,6 +23,8 @@ public class Localization : ModuleRules
 				"InputCore",
 				"SourceControl",
 				"Json",
+				"JsonUtilities",
+				"Projects",
 			}
 		);
 
diff --git a/Engine/Source/Developer/Localization/Private/UserGeneratedContentLocalization.cpp b/Engine/Source/Developer/Localization/Private/UserGeneratedContentLocalization.cpp
new file mode 100644
--- /dev/null
+++ b/Engine/Source/Developer/Localization/Private/UserGeneratedContentLocalization.cpp
@@ -0,0 +1,671 @@
+// Copyright Epic Games, Inc. All Rights Reserved.
+
+#include "UserGeneratedContentLocalization.h"
+
+#include "HAL/FileManager.h"
+#include "Misc/FileHelper.h"
+#include "Misc/PathViews.h"
+#include "Misc/ScopeExit.h"
+#include "JsonObjectConverter.h"
+#include "Interfaces/IPluginManager.h"
+
+#include "Internationalization/CultureFilter.h"
+#include "Internationalization/TextLocalizationManager.h"
+#include "TextLocalizationResourceGenerator.h"
+#include "LocalizationConfigurationScript.h"
+#include "LocTextHelper.h"
+#include "PortableObjectFormatDOM.h"
+#include "SourceControlHelpers.h"
+
+void FUserGeneratedContentLocalizationDescriptor::InitializeFromProject(const ELocalizedTextSourceCategory LocalizationCategory)
+{
+	ELocalizationLoadFlags LoadFlags = ELocalizationLoadFlags::None;
+	switch (LocalizationCategory)
+	{
+	case ELocalizedTextSourceCategory::Game:
+		LoadFlags |= ELocalizationLoadFlags::Game;
+		break;
+
+	case ELocalizedTextSourceCategory::Engine:
+		LoadFlags |= ELocalizationLoadFlags::Engine;
+		break;
+
+	case ELocalizedTextSourceCategory::Editor:
+		LoadFlags |= ELocalizationLoadFlags::Editor;
+		break;
+
+	default:
+		checkf(false, TEXT("Unexpected ELocalizedTextSourceCategory!"));
+		break;
+	}
+
+	NativeCulture = FTextLocalizationManager::Get().GetNativeCultureName(LocalizationCategory);
+	CulturesToGenerate = FTextLocalizationManager::Get().GetLocalizedCultureNames(LoadFlags);
+
+	// Filter any cultures that are disabled in shipping or via UGC loc settings
+	{
+		const FCultureFilter CultureFilter(EBuildConfiguration::Shipping, ELocalizationLoadFlags::Engine | LoadFlags);
+		CulturesToGenerate.RemoveAll([&CultureFilter](const FString& Culture)
+		{
+			return !CultureFilter.IsCultureAllowed(Culture)
+				|| GetDefault<UUserGeneratedContentLocalizationSettings>()->CulturesToDisable.Contains(Culture);
+		});
+	}
+}
+
+bool FUserGeneratedContentLocalizationDescriptor::Validate(const FUserGeneratedContentLocalizationDescriptor& DefaultDescriptor)
+{
+	int32 NumCulturesFixed = 0;
+
+	if (!DefaultDescriptor.CulturesToGenerate.Contains(NativeCulture))
+	{
+		++NumCulturesFixed;
+		NativeCulture = DefaultDescriptor.NativeCulture;
+	}
+
+	NumCulturesFixed += CulturesToGenerate.RemoveAll([&DefaultDescriptor](const FString& Culture)
+	{
+		return !DefaultDescriptor.CulturesToGenerate.Contains(Culture);
+	});
+
+	return NumCulturesFixed == 0;
+}
+
+bool FUserGeneratedContentLocalizationDescriptor::ToJsonObject(TSharedPtr<FJsonObject>& OutJsonObject) const
+{
+	OutJsonObject = FJsonObjectConverter::UStructToJsonObject(*this);
+	return OutJsonObject.IsValid();
+}
+
+bool FUserGeneratedContentLocalizationDescriptor::ToJsonString(FString& OutJsonString) const
+{
+	return FJsonObjectConverter::UStructToJsonObjectString(*this, OutJsonString);
+}
+
+bool FUserGeneratedContentLocalizationDescriptor::ToJsonFile(const TCHAR* InFilename) const
+{
+	FString UGCLocDescData;
+	return ToJsonString(UGCLocDescData) 
+		&& FFileHelper::SaveStringToFile(UGCLocDescData, InFilename, FFileHelper::EEncodingOptions::ForceUTF8);
+}
+
+bool FUserGeneratedContentLocalizationDescriptor::FromJsonObject(TSharedRef<const FJsonObject> InJsonObject)
+{
+	return FJsonObjectConverter::JsonObjectToUStruct(ConstCastSharedRef<FJsonObject>(InJsonObject), this);
+}
+
+bool FUserGeneratedContentLocalizationDescriptor::FromJsonString(const FString& InJsonString)
+{
+	return FJsonObjectConverter::JsonObjectStringToUStruct(InJsonString, this);
+}
+
+bool FUserGeneratedContentLocalizationDescriptor::FromJsonFile(const TCHAR* InFilename)
+{
+	FString UGCLocDescData;
+	return FFileHelper::LoadFileToString(UGCLocDescData, InFilename)
+		&& FromJsonString(UGCLocDescData);
+}
+
+namespace UserGeneratedContentLocalization
+{
+
+FString GetLocalizationScratchDirectory()
+{
+	return FPaths::Combine(FPaths::ProjectIntermediateDir(), TEXT("Localization"));
+}
+
+FString GetLocalizationScratchDirectory(const TSharedRef<IPlugin>& Plugin)
+{
+	return FPaths::Combine(GetLocalizationScratchDirectory(), Plugin->GetName());
+}
+
+FString GetLocalizationTargetDirectory(const FString& PluginName, const FString& PluginContentDirectory)
+{
+	return FPaths::Combine(PluginContentDirectory, TEXT("Localization"), PluginName);
+}
+
+FString GetLocalizationTargetDirectory(const TSharedRef<IPlugin>& Plugin)
+{
+	return GetLocalizationTargetDirectory(Plugin->GetName(), Plugin->GetContentDir());
+}
+
+void PreWriteFileWithSCC(const FString& Filename)
+{
+	if (USourceControlHelpers::IsEnabled())
+	{
+		// If the file already already exists, then check it out before writing to it
+		if (FPaths::FileExists(Filename) && !USourceControlHelpers::CheckOutFile(Filename))
+		{
+			UE_LOG(LogLocalization, Error, TEXT("Failed to check out file '%s'. %s"), *Filename, *USourceControlHelpers::LastErrorMsg().ToString());
+		}
+	}
+}
+
+void PostWriteFileWithSCC(const FString& Filename)
+{
+	if (USourceControlHelpers::IsEnabled())
+	{
+		// If the file didn't exist before then this will add it, otherwise it will do nothing
+		if (!USourceControlHelpers::CheckOutOrAddFile(Filename))
+		{
+			UE_LOG(LogLocalization, Error, TEXT("Failed to check out file '%s'. %s"), *Filename, *USourceControlHelpers::LastErrorMsg().ToString());
+		}
+	}
+}
+
+bool ExportLocalization(TArrayView<const TSharedRef<IPlugin>> Plugins, const FExportLocalizationOptions& ExportOptions, TFunctionRef<int32(const FString&, FString&)> CommandletExecutor)
+{
+	if (ExportOptions.UGCLocDescriptor.NativeCulture.IsEmpty())
+	{
+		UE_LOG(LogLocalization, Error, TEXT("Localization export options did not have a 'NativeCulture' set"));
+		return false;
+	}
+
+	// Create a scratch directory for the temporary localization data
+	const FString RootLocalizationScratchDirectory = GetLocalizationScratchDirectory();
+	IFileManager::Get().MakeDirectory(*RootLocalizationScratchDirectory, /*bTree*/true);
+	ON_SCOPE_EXIT
+	{
+		if (ExportOptions.bAutoCleanup)
+		{
+			// Delete the entire scratch directory
+			IFileManager::Get().DeleteDirectory(*RootLocalizationScratchDirectory, /*bRequireExists*/false, /*bTree*/true);
+		}
+	};
+
+	// Make sure we're also exporting localization for the native culture
+	TArray<FString> CulturesToGenerate = ExportOptions.UGCLocDescriptor.CulturesToGenerate;
+	CulturesToGenerate.AddUnique(ExportOptions.UGCLocDescriptor.NativeCulture);
+
+	// Localization data stored per-plugin
+	TArray<FString, TInlineAllocator<1>> GatherConfigFilenames;
+	for (const TSharedRef<IPlugin>& Plugin : Plugins)
+	{
+		const FString PluginLocalizationScratchDirectory = GetLocalizationScratchDirectory(Plugin);
+		const FString PluginLocalizationTargetDirectory = GetLocalizationTargetDirectory(Plugin);
+
+		// Seed the scratch directory with the current PO files for this plugin, so that the loc gather will import and preserve any existing translation data
+		{
+			bool bCopiedAllFiles = true;
+			IFileManager::Get().IterateDirectoryRecursively(*PluginLocalizationTargetDirectory, [&Plugin, &PluginLocalizationScratchDirectory, &PluginLocalizationTargetDirectory, &bCopiedAllFiles](const TCHAR* FilenameOrDirectory, bool bIsDirectory)
+			{
+				if (!bIsDirectory && FPathViews::GetExtension(FilenameOrDirectory) == TEXTVIEW("po"))
+				{
+					FString DestinationFilename = FilenameOrDirectory;
+					if (DestinationFilename.ReplaceInline(*PluginLocalizationTargetDirectory, *PluginLocalizationScratchDirectory) > 0)
+					{
+						IFileManager::Get().Copy(*DestinationFilename, FilenameOrDirectory);
+					}
+					else
+					{
+						bCopiedAllFiles = false;
+						return false;
+					}
+				}
+				return true;
+			});
+			if (!bCopiedAllFiles)
+			{
+				UE_LOG(LogLocalization, Error, TEXT("Failed to import existing .po files for '%s'"), *Plugin->GetName());
+				return false;
+			}
+		}
+
+		// Build the gather config
+		{
+			// Build up a basic localization config that will do the following:
+			//  1) Gather source/assets in the current plugin
+			//  2) Import any existing PO file data
+			//  3) Export new PO file data
+
+			FLocalizationConfigurationScript GatherConfig;
+			int32 GatherStepIndex = 0;
+
+			// Common
+			{
+				FConfigSection& ConfigSection = GatherConfig.CommonSettings();
+
+				ConfigSection.Add(TEXT("SourcePath"), FPaths::ConvertRelativePathToFull(PluginLocalizationScratchDirectory));
+				ConfigSection.Add(TEXT("DestinationPath"), FPaths::ConvertRelativePathToFull(PluginLocalizationScratchDirectory));
+
+				ConfigSection.Add(TEXT("ManifestName"), FString::Printf(TEXT("%s.manifest"), *Plugin->GetName()));
+				ConfigSection.Add(TEXT("ArchiveName"), FString::Printf(TEXT("%s.archive"), *Plugin->GetName()));
+				ConfigSection.Add(TEXT("PortableObjectName"), FString::Printf(TEXT("%s.po"), *Plugin->GetName()));
+
+				ConfigSection.Add(TEXT("NativeCulture"), ExportOptions.UGCLocDescriptor.NativeCulture);
+				for (const FString& CultureToGenerate : CulturesToGenerate)
+				{
+					ConfigSection.Add(TEXT("CulturesToGenerate"), *CultureToGenerate);
+				}
+			}
+
+			// Gather source
+			if (ExportOptions.bGatherSource)
+			{
+				const FString PluginConfigDir = FPaths::ConvertRelativePathToFull(FPaths::Combine(Plugin->GetBaseDir(), TEXT("Config")));
+				const FString PluginSourceDir = FPaths::ConvertRelativePathToFull(FPaths::Combine(Plugin->GetBaseDir(), TEXT("Source")));
+
+				TArray<FString, TInlineAllocator<2>> SearchDirectoryPaths;
+				if (FPaths::DirectoryExists(PluginConfigDir))
+				{
+					SearchDirectoryPaths.Add(PluginConfigDir);
+				}
+				if (FPaths::DirectoryExists(PluginSourceDir))
+				{
+					SearchDirectoryPaths.Add(PluginSourceDir);
+				}
+
+				// Only gather from source if there's valid paths to gather from, as otherwise the commandlet will error
+				if (SearchDirectoryPaths.Num() > 0)
+				{
+					FConfigSection& ConfigSection = GatherConfig.GatherTextStep(GatherStepIndex++);
+					ConfigSection.Add(TEXT("CommandletClass"), TEXT("GatherTextFromSource"));
+
+					ConfigSection.Add(TEXT("FileNameFilters"), TEXT("*.h"));
+					ConfigSection.Add(TEXT("FileNameFilters"), TEXT("*.cpp"));
+					ConfigSection.Add(TEXT("FileNameFilters"), TEXT("*.inl"));
+					ConfigSection.Add(TEXT("FileNameFilters"), TEXT("*.ini"));
+
+					for (const FString& SearchDirectoryPath : SearchDirectoryPaths)
+					{
+						ConfigSection.Add(TEXT("SearchDirectoryPaths"), SearchDirectoryPath);
+					}
+				}
+			}
+
+			// Gather assets
+			if (ExportOptions.bGatherAssets && Plugin->CanContainContent())
+			{
+				FConfigSection& ConfigSection = GatherConfig.GatherTextStep(GatherStepIndex++);
+				ConfigSection.Add(TEXT("CommandletClass"), TEXT("GatherTextFromAssets"));
+
+				ConfigSection.Add(TEXT("PackageFileNameFilters"), TEXT("*.uasset"));
+				ConfigSection.Add(TEXT("PackageFileNameFilters"), TEXT("*.umap"));
+
+				ConfigSection.Add(TEXT("IncludePathFilters"), FPaths::ConvertRelativePathToFull(FPaths::Combine(Plugin->GetContentDir(), TEXT("*"))));
+
+				ConfigSection.Add(TEXT("ExcludePathFilters"), FPaths::ConvertRelativePathToFull(FPaths::Combine(Plugin->GetContentDir(), TEXT("Localization"), TEXT("*"))));
+				ConfigSection.Add(TEXT("ExcludePathFilters"), FPaths::ConvertRelativePathToFull(FPaths::Combine(Plugin->GetContentDir(), TEXT("L10N"), TEXT("*"))));
+
+				if (const FString* CollectionFilter = ExportOptions.PluginNameToCollectionNameFilter.Find(Plugin->GetName()))
+				{
+					ConfigSection.Add(TEXT("CollectionFilters"), *CollectionFilter);
+				}
+			}
+
+			// Generate manifest
+			{
+				FConfigSection& ConfigSection = GatherConfig.GatherTextStep(GatherStepIndex++);
+				ConfigSection.Add(TEXT("CommandletClass"), TEXT("GenerateGatherManifest"));
+			}
+
+			// Generate archive
+			{
+				FConfigSection& ConfigSection = GatherConfig.GatherTextStep(GatherStepIndex++);
+				ConfigSection.Add(TEXT("CommandletClass"), TEXT("GenerateGatherArchive"));
+			}
+
+			// Import PO
+			{
+				// Read the UGC localization descriptor settings that were used to generate this localization data, as we should import against those
+				FUserGeneratedContentLocalizationDescriptor UGCLocDescriptorForImport;
+				{
+					const FString UGCLocFilename = FPaths::Combine(PluginLocalizationTargetDirectory, FString::Printf(TEXT("%s.ugcloc"), *Plugin->GetName()));
+					if (!UGCLocDescriptorForImport.FromJsonFile(*UGCLocFilename))
+					{
+						UGCLocDescriptorForImport = ExportOptions.UGCLocDescriptor;
+					}
+				}
+
+				FConfigSection& ConfigSection = GatherConfig.GatherTextStep(GatherStepIndex++);
+				ConfigSection.Add(TEXT("CommandletClass"), TEXT("InternationalizationExport"));
+
+				ConfigSection.Add(TEXT("bImportLoc"), TEXT("true"));
+
+				ConfigSection.Add(TEXT("POFormat"), StaticEnum<EPortableObjectFormat>()->GetNameStringByValue((int64)UGCLocDescriptorForImport.PoFormat));
+			}
+
+			// Export PO
+			{
+				FConfigSection& ConfigSection = GatherConfig.GatherTextStep(GatherStepIndex++);
+				ConfigSection.Add(TEXT("CommandletClass"), TEXT("InternationalizationExport"));
+
+				ConfigSection.Add(TEXT("bExportLoc"), TEXT("true"));
+
+				ConfigSection.Add(TEXT("POFormat"), StaticEnum<EPortableObjectFormat>()->GetNameStringByValue((int64)ExportOptions.UGCLocDescriptor.PoFormat));
+
+				ConfigSection.Add(TEXT("ShouldPersistCommentsOnExport"), TEXT("true"));
+			}
+
+			// Write config
+			{
+				GatherConfig.Dirty = true;
+
+				FString GatherConfigFilename = FPaths::ConvertRelativePathToFull(RootLocalizationScratchDirectory / FString::Printf(TEXT("%s.ini"), *Plugin->GetName()));
+				if (GatherConfig.Write(GatherConfigFilename))
+				{
+					GatherConfigFilenames.Add(MoveTemp(GatherConfigFilename));
+				}
+				else
+				{
+					UE_LOG(LogLocalization, Error, TEXT("Failed to write gather config for '%s': %s"), *Plugin->GetName(), *GatherConfigFilename);
+					return false;
+				}
+			}
+		}
+	}
+	
+	// Run the commandlet
+	if (GatherConfigFilenames.Num() > 0)
+	{
+		FString CommandletOutput;
+		const int32 ReturnCode = CommandletExecutor(FString::Join(GatherConfigFilenames, TEXT(";")), CommandletOutput);
+
+		// Log the output and result of the commandlet
+		{
+			TArray<FString> CommandletOutputLines;
+			CommandletOutput.ParseIntoArrayLines(CommandletOutputLines);
+
+			UE_LOG(LogLocalization, Log, TEXT("Localization commandlet finished with exit code %d"), ReturnCode);
+			for (const FString& CommandletOutputLine : CommandletOutputLines)
+			{
+				UE_LOG(LogLocalization, Log, TEXT("    %s"), *CommandletOutputLine);
+			}
+		}
+
+		// Verify the commandlet finished cleanly
+		if (ReturnCode != 0)
+		{
+			return false;
+		}
+	}
+
+	// Copy any updated PO files back to the plugins and write out the localization settings used to generate them
+	for (const TSharedRef<IPlugin>& Plugin : Plugins)
+	{
+		const FString PluginLocalizationScratchDirectory = GetLocalizationScratchDirectory(Plugin);
+		const FString PluginLocalizationTargetDirectory = GetLocalizationTargetDirectory(Plugin);
+
+		// Write the UGC localization descriptor settings that were used to generate this localization data
+		// That will be needed to handle compilation, but also to handle import correctly if the descriptor settings change later
+		{
+			const FString UGCLocFilename = FPaths::Combine(PluginLocalizationTargetDirectory, FString::Printf(TEXT("%s.ugcloc"), *Plugin->GetName()));
+
+			PreWriteFileWithSCC(UGCLocFilename);
+			if (!ExportOptions.UGCLocDescriptor.ToJsonFile(*UGCLocFilename))
+			{
+				PostWriteFileWithSCC(UGCLocFilename);
+			}
+			else
+			{
+				UE_LOG(LogLocalization, Warning, TEXT("Failed to update .ugcloc file for '%s': %s"), *Plugin->GetName(), *UGCLocFilename);
+			}
+		}
+
+		IFileManager::Get().IterateDirectoryRecursively(*PluginLocalizationScratchDirectory, [&Plugin, &PluginLocalizationScratchDirectory, &PluginLocalizationTargetDirectory](const TCHAR* FilenameOrDirectory, bool bIsDirectory)
+		{
+			if (!bIsDirectory && FPathViews::GetExtension(FilenameOrDirectory) == TEXTVIEW("po"))
+			{
+				FString DestinationFilename = FilenameOrDirectory;
+				if (DestinationFilename.ReplaceInline(*PluginLocalizationScratchDirectory, *PluginLocalizationTargetDirectory) > 0)
+				{
+					PreWriteFileWithSCC(DestinationFilename);
+					if (IFileManager::Get().Copy(*DestinationFilename, FilenameOrDirectory) == COPY_OK)
+					{
+						PostWriteFileWithSCC(DestinationFilename);
+					}
+					else
+					{
+						UE_LOG(LogLocalization, Warning, TEXT("Failed to update .po file for '%s': %s"), *Plugin->GetName(), *DestinationFilename);
+					}
+				}
+			}
+			return true;
+		});
+
+		if (ExportOptions.bUpdatePluginDescriptor)
+		{
+			FPluginDescriptor PluginDescriptor = Plugin->GetDescriptor();
+			if (!PluginDescriptor.LocalizationTargets.ContainsByPredicate([&Plugin](const FLocalizationTargetDescriptor& LocalizationTargetDescriptor) { return LocalizationTargetDescriptor.Name == Plugin->GetName(); }))
+			{
+				FLocalizationTargetDescriptor& LocalizationTargetDescriptor = PluginDescriptor.LocalizationTargets.AddDefaulted_GetRef();
+				LocalizationTargetDescriptor.Name = Plugin->GetName();
+				switch (ExportOptions.LocalizationCategory)
+				{
+				case ELocalizedTextSourceCategory::Game:
+					LocalizationTargetDescriptor.LoadingPolicy = ELocalizationTargetDescriptorLoadingPolicy::Game;
+					break;
+
+				case ELocalizedTextSourceCategory::Engine:
+					LocalizationTargetDescriptor.LoadingPolicy = ELocalizationTargetDescriptorLoadingPolicy::Always;
+					break;
+
+				case ELocalizedTextSourceCategory::Editor:
+					LocalizationTargetDescriptor.LoadingPolicy = ELocalizationTargetDescriptorLoadingPolicy::Editor;
+					break;
+
+				default:
+					checkf(false, TEXT("Unexpected ELocalizedTextSourceCategory!"));
+					break;
+				}
+
+				FText DescriptorUpdateFailureReason;
+				PreWriteFileWithSCC(Plugin->GetDescriptorFileName());
+				if (Plugin->UpdateDescriptor(PluginDescriptor, DescriptorUpdateFailureReason))
+				{
+					PostWriteFileWithSCC(Plugin->GetDescriptorFileName());
+				}
+				else
+				{
+					UE_LOG(LogLocalization, Warning, TEXT("Failed to update .uplugin file for '%s': %s"), *Plugin->GetName(), *DescriptorUpdateFailureReason.ToString());
+				}
+			}
+		}
+	}
+
+	return true;
+}
+
+bool CompileLocalizationTarget(const FString& LocalizationTargetDirectory, const FLocTextHelper& LocTextHelper)
+{
+	const FString LocMetaName = FString::Printf(TEXT("%s.locmeta"), *LocTextHelper.GetTargetName());
+	const FString LocResName = FString::Printf(TEXT("%s.locres"), *LocTextHelper.GetTargetName());
+
+	// Generate the LocMeta file
+	{
+		FTextLocalizationMetaDataResource LocMeta;
+		if (FTextLocalizationResourceGenerator::GenerateLocMeta(LocTextHelper, LocResName, LocMeta))
+		{
+			if (!LocMeta.SaveToFile(LocalizationTargetDirectory / LocMetaName))
+			{
+				UE_LOG(LogLocalization, Error, TEXT("Failed to save LocMeta file for '%s'"), *LocTextHelper.GetTargetName());
+				return false;
+			}
+		}
+		else
+		{
+			UE_LOG(LogLocalization, Error, TEXT("Failed to generate LocMeta file for '%s'"), *LocTextHelper.GetTargetName());
+			return false;
+		}
+	}
+
+	// Generate the LocRes files
+	for (const FString& CultureToGenerate : LocTextHelper.GetAllCultures())
+	{
+		FTextLocalizationResource LocRes;
+		TMap<FName, TSharedRef<FTextLocalizationResource>> PerPlatformLocRes;
+		if (FTextLocalizationResourceGenerator::GenerateLocRes(LocTextHelper, CultureToGenerate, EGenerateLocResFlags::None, LocalizationTargetDirectory / CultureToGenerate / LocResName, LocRes, PerPlatformLocRes))
+		{
+			checkf(PerPlatformLocRes.Num() == 0, TEXT("UGC localization does not support per-platform LocRes!"));
+
+			if (!LocRes.SaveToFile(LocalizationTargetDirectory / CultureToGenerate / LocResName))
+			{
+				UE_LOG(LogLocalization, Error, TEXT("Failed to save LocRes file for '%s' (culture '%s')"), *LocTextHelper.GetTargetName(), *CultureToGenerate);
+				return false;
+			}
+		}
+		else
+		{
+			UE_LOG(LogLocalization, Error, TEXT("Failed to generate LocRes file for '%s' (culture '%s')"), *LocTextHelper.GetTargetName(), *CultureToGenerate);
+			return false;
+		}
+	}
+
+	return true;
+}
+
+bool CompileLocalization(TArrayView<const TSharedRef<IPlugin>> Plugins, const FUserGeneratedContentLocalizationDescriptor* DefaultDescriptor)
+{
+	// Localization data is stored per-plugin
+	for (const TSharedRef<IPlugin>& Plugin : Plugins)
+	{
+		const FString PluginLocalizationTargetDirectory = GetLocalizationTargetDirectory(Plugin);
+		if (!CompileLocalization(Plugin->GetName(), PluginLocalizationTargetDirectory, PluginLocalizationTargetDirectory, DefaultDescriptor))
+		{
+			return false;
+		}
+	}
+
+	return true;
+}
+
+bool CompileLocalization(const FString& PluginName, const FString& PluginInputContentDirectory, const FString& PluginOutputContentDirectory, const FUserGeneratedContentLocalizationDescriptor* DefaultDescriptor)
+{
+	// Load the localization data so that we can compile it
+	TSharedPtr<FLocTextHelper> LocTextHelper;
+	const ELoadLocalizationResult LoadResult = LoadLocalization(PluginName, PluginInputContentDirectory, LocTextHelper, DefaultDescriptor);
+	if (LoadResult == ELoadLocalizationResult::NoData)
+	{
+		// Nothing to do
+		return true;
+	}
+	if (LoadResult == ELoadLocalizationResult::Failed)
+	{
+		// Failed to load, so can't compile
+		return false;
+	}
+	check(LoadResult == ELoadLocalizationResult::Success);
+
+	const FString PluginLocalizationOutputDirectory = GetLocalizationTargetDirectory(PluginName, PluginOutputContentDirectory);
+	return CompileLocalizationTarget(PluginLocalizationOutputDirectory, *LocTextHelper);
+}
+
+bool ImportPortableObject(const FString& LocalizationTargetDirectory, const FString& CultureToLoad, const EPortableObjectFormat PoFormat, FLocTextHelper& LocTextHelper)
+{
+	const bool bIsNativeCulture = CultureToLoad == LocTextHelper.GetNativeCulture();
+
+	const FString POFilename = LocalizationTargetDirectory / CultureToLoad / FString::Printf(TEXT("%s.po"), *LocTextHelper.GetTargetName());
+
+	FString POFileData;
+	FPortableObjectFormatDOM POFile;
+	if (!FFileHelper::LoadFileToString(POFileData, *POFilename) || !POFile.FromString(POFileData))
+	{
+		return false;
+	}
+
+	// Process each PO entry
+	for (auto EntryPairIter = POFile.GetEntriesIterator(); EntryPairIter; ++EntryPairIter)
+	{
+		const TSharedPtr<FPortableObjectEntry>& POEntry = EntryPairIter->Value;
+		if (POEntry->MsgId.IsEmpty() || POEntry->MsgStr.Num() == 0 || POEntry->MsgStr[0].IsEmpty())
+		{
+			// We ignore the header entry or entries with no translation.
+			continue;
+		}
+
+		FString Namespace;
+		FString Key;
+		FString SourceText;
+		FString Translation;
+		PortableObjectPipeline::ParseBasicPOFileEntry(*POEntry, Namespace, Key, SourceText, Translation, ELocalizedTextCollapseMode::IdenticalTextIdAndSource, PoFormat);
+
+		// PO files don't contain the key meta-data so we can't reconstruct this
+		// Key meta-data only exists to force the PO file export an ID that contains both the namespace AND key though, so it doesn't matter if it's lost here as it won't affect the LocRes generation
+		TSharedPtr<FLocMetadataObject> KeyMetadataObj = nullptr;
+
+		// Not all formats contain the source string, so if the source is empty then 
+		// we'll assume the translation was made against the most up-to-date source
+		if (SourceText.IsEmpty())
+		{
+			if (bIsNativeCulture)
+			{
+				SourceText = Translation;
+			}
+			else if (const TSharedPtr<FArchiveEntry> NativeEntry = LocTextHelper.FindTranslation(LocTextHelper.GetNativeCulture(), Namespace, Key, KeyMetadataObj))
+			{
+				SourceText = NativeEntry->Translation.Text;
+			}
+		}
+
+		// If this is the native culture then we also add it as source in the manifest
+		if (bIsNativeCulture)
+		{
+			FManifestContext ManifestContext;
+			ManifestContext.SourceLocation = POEntry->ReferenceComments.Num() > 0 ? POEntry->ReferenceComments[0] : FString();
+			ManifestContext.Key = Key;
+			ManifestContext.KeyMetadataObj = KeyMetadataObj;
+			LocTextHelper.AddSourceText(Namespace, FLocItem(SourceText), ManifestContext);
+		}
+
+		// All cultures add this info as a translation
+		LocTextHelper.AddTranslation(CultureToLoad, Namespace, Key, KeyMetadataObj, FLocItem(SourceText), FLocItem(Translation), /*bIsOptional*/false);
+	}
+
+	return true;
+}
+
+ELoadLocalizationResult LoadLocalization(const FString& PluginName, const FString& PluginContentDirectory, TSharedPtr<FLocTextHelper>& OutLocTextHelper, const FUserGeneratedContentLocalizationDescriptor* DefaultDescriptor)
+{
+	const FString PluginLocalizationTargetDirectory = GetLocalizationTargetDirectory(PluginName, PluginContentDirectory);
+
+	const FString UGCLocFilename = FPaths::Combine(PluginLocalizationTargetDirectory, FString::Printf(TEXT("%s.ugcloc"), *PluginName));
+	if (!FPaths::FileExists(UGCLocFilename))
+	{
+		// Nothing to do
+		return ELoadLocalizationResult::NoData;
+	}
+
+	// Read the UGC localization descriptor settings that were used to generate this localization data
+	FUserGeneratedContentLocalizationDescriptor UGCLocDescriptor;
+	if (!UGCLocDescriptor.FromJsonFile(*UGCLocFilename))
+	{
+		UE_LOG(LogLocalization, Error, TEXT("Failed to load localization descriptor for '%s'"), *PluginName);
+		return ELoadLocalizationResult::Failed;
+	}
+
+	// Validate the loaded settings against the given default
+	// This will remove/reset any invalid data
+	if (DefaultDescriptor)
+	{
+		UGCLocDescriptor.Validate(*DefaultDescriptor);
+	}
+
+	// Create in-memory versions of the manifest/archives that we will populate below
+	OutLocTextHelper = MakeShared<FLocTextHelper>(PluginLocalizationTargetDirectory, FString::Printf(TEXT("%s.manifest"), *PluginName), FString::Printf(TEXT("%s.archive"), *PluginName), UGCLocDescriptor.NativeCulture, UGCLocDescriptor.CulturesToGenerate, nullptr);
+	OutLocTextHelper->LoadAll(ELocTextHelperLoadFlags::Create);
+
+	// Import each PO file data, as we'll use it to generate the LocRes (via FLocTextHelper)
+	// We always process the native culture first as it's also used to populate the manifest with the source texts
+	if (!ImportPortableObject(PluginLocalizationTargetDirectory, OutLocTextHelper->GetNativeCulture(), UGCLocDescriptor.PoFormat, *OutLocTextHelper))
+	{
+		UE_LOG(LogLocalization, Error, TEXT("Failed to load PO file for '%s' (culture '%s')"), *PluginName, *OutLocTextHelper->GetNativeCulture());
+		return ELoadLocalizationResult::Failed;
+	}
+	for (const FString& CultureToGenerate : OutLocTextHelper->GetAllCultures())
+	{
+		if (CultureToGenerate == OutLocTextHelper->GetNativeCulture())
+		{
+			continue;
+		}
+
+		if (!ImportPortableObject(PluginLocalizationTargetDirectory, CultureToGenerate, UGCLocDescriptor.PoFormat, *OutLocTextHelper))
+		{
+			UE_LOG(LogLocalization, Error, TEXT("Failed to load PO file for '%s' (culture '%s')"), *PluginName, *CultureToGenerate);
+			return ELoadLocalizationResult::Failed;
+		}
+	}
+
+	return ELoadLocalizationResult::Success;
+}
+
+}
diff --git a/Engine/Source/Developer/Localization/Public/UserGeneratedContentLocalization.h b/Engine/Source/Developer/Localization/Public/UserGeneratedContentLocalization.h
new file mode 100644
--- /dev/null
+++ b/Engine/Source/Developer/Localization/Public/UserGeneratedContentLocalization.h
@@ -0,0 +1,196 @@
+// Copyright Epic Games, Inc. All Rights Reserved.
+
+#pragma once
+
+#include "Containers/Array.h"
+#include "Containers/ArrayView.h"
+#include "Containers/Map.h"
+#include "Containers/UnrealString.h"
+#include "Templates/Function.h"
+#include "Templates/SharedPointer.h"
+
+#include "PortableObjectPipeline.h"
+#include "Internationalization/LocalizedTextSourceTypes.h"
+
+#include "UserGeneratedContentLocalization.generated.h"
+
+class IPlugin;
+class FJsonObject;
+class FLocTextHelper;
+
+/**
+ * Settings controlling UGC localization.
+ */
+UCLASS(config=Engine, defaultconfig)
+class LOCALIZATION_API UUserGeneratedContentLocalizationSettings : public UObject
+{
+	GENERATED_BODY()
+
+public:
+	/**
+	 * List of cultures that should be disabled for UGC localization.
+	 * @note You can't disable the native culture for the project.
+	 */
+	UPROPERTY(EditAnywhere, Category=Localization)
+	TArray<FString> CulturesToDisable;
+
+	/**
+	 * Should we compile UGC localization (if present) for DLC plugins during cook?
+	 */
+	UPROPERTY(EditAnywhere, Category=Localization)
+	bool bCompileDLCLocalizationDuringCook = true;
+
+	/**
+	 * Should we validate UGC localization (if present) for DLC plugins during cook?
+	 * @note Validation will happen against a UGC localization descriptor that has had InitializeFromProject called on it.
+	 */
+	UPROPERTY(EditAnywhere, Category=Localization)
+	bool bValidateDLCLocalizationDuringCook = true;
+};
+
+/**
+ * Minimal descriptor needed to generate a localization target for UGC localization.
+ */
+USTRUCT()
+struct LOCALIZATION_API FUserGeneratedContentLocalizationDescriptor
+{
+	GENERATED_BODY()
+	
+public:
+	/**
+	 * Initialize the NativeCulture and CulturesToGenerate values based on the settings of the currently loaded Unreal project.
+	 * @param LocalizationCategory What category is the localization targets being used with this descriptor?
+	 */
+	void InitializeFromProject(const ELocalizedTextSourceCategory LocalizationCategory = ELocalizedTextSourceCategory::Game);
+	
+	/**
+	 * Validate that this descriptor isn't using cultures that aren't present in the CulturesToGenerate of the given default.
+	 *   - If the NativeCulture is invalid, reset it to the value from the default.
+	 *   - If CulturesToGenerate contains invalid entries then remove those from the array.
+	 * 
+	 * @return True if this descriptor was valid and no changes were made. False if this descriptor was invalid and had default changes applied.
+	 */
+	bool Validate(const FUserGeneratedContentLocalizationDescriptor& DefaultDescriptor);
+
+	/**
+	 * Save the settings to a JSON object/file.
+	 */
+	bool ToJsonObject(TSharedPtr<FJsonObject>& OutJsonObject) const;
+	bool ToJsonString(FString& OutJsonString) const;
+	bool ToJsonFile(const TCHAR* InFilename) const;
+
+	/**
+	 * Load the settings from a JSON object/file.
+	 */
+	bool FromJsonObject(TSharedRef<const FJsonObject> InJsonObject);
+	bool FromJsonString(const FString& InJsonString);
+	bool FromJsonFile(const TCHAR* InFilename);
+	
+	/**
+	 * The culture that the source text is authored in.
+	 * @note You shouldn't change this once you start to localize your text.
+	 */
+	UPROPERTY(EditAnywhere, Category=Localization)
+	FString NativeCulture;
+
+	/**
+	 * The cultures that we should generate localization data for.
+	 * @note Will implicitly always contain the native culture during export/compile.
+	 */
+	UPROPERTY(EditAnywhere, Category=Localization)
+	TArray<FString> CulturesToGenerate;
+
+	/**
+	 * What format of PO file should we use?
+	 * @note You can adjust this later and we'll attempt to preserve any existing localization data by importing with the old setting prior to export.
+	 */
+	UPROPERTY(EditAnywhere, AdvancedDisplay, Category=Localization, meta=(DisplayName="PO Format"))
+	EPortableObjectFormat PoFormat = EPortableObjectFormat::Unreal;
+};
+
+/**
+ * UGC localization can be used to provide a simplified localization experience for basic plugins (only providing the PO files to be translated), 
+ * and is primarily designed for DLC plugins where the UGC localization will be compiled during cook (@see UUserGeneratedContentLocalizationSettings).
+ * 
+ * Support for non-DLC plugins can be provided via project specific tooling built upon this base API.
+ * Support for complex plugins (such as those containing different kinds of modules, eg) a mix of game/engine and editor) are not supported via this API.
+ */
+namespace UserGeneratedContentLocalization
+{
+
+struct FExportLocalizationOptions
+{
+	/** Common export options for all plugins */
+	FUserGeneratedContentLocalizationDescriptor UGCLocDescriptor;
+
+	/** Optional mapping of plugin names to collection names (to act as a filter for their asset gather step) */
+	TMap<FString, FString> PluginNameToCollectionNameFilter;
+
+	/** True to gather localization from source code (if a plugin has a Config or Source folder) */
+	bool bGatherSource = true;
+
+	/** True to gather localization from assets */
+	bool bGatherAssets = true;
+
+	/** True to update the plugin descriptors (if needed) so that they contain the exported localization target */
+	bool bUpdatePluginDescriptor = true;
+
+	/** True to automatically clean-up any scratch data created during the localization export */
+	bool bAutoCleanup = true;
+
+	/** The category to use for the exported localization target (only used when bUpdatePluginDescriptor is true) */
+	ELocalizedTextSourceCategory LocalizationCategory = ELocalizedTextSourceCategory::Game;
+};
+
+enum class ELoadLocalizationResult : uint8
+{
+	/** There was no source localization data to load */
+	NoData,
+	/** There was source localization data to load, but we failed to load it */
+	Failed,
+	/** There was source localization data to load, and we successfully loaded it */
+	Success,
+};
+
+/**
+ * Export UGC localization for the given plugins.
+ * 
+ * @param Plugins				The list of plugins to export.
+ * @param ExportOptions			Options controlling how to export the localization data.
+ * @param CommandletExecutor	Callback used to actually execute the gather commandlet:
+ *									This should execute an editor with `-run=GatherText -config="..."`, where the config argument is the first argument passed to this callback.
+ *									The second argument should be filled with the raw log output from running the commandlet process.
+ *									The return value is the exit code of the commandlet process (where zero means success).
+ */
+LOCALIZATION_API bool ExportLocalization(TArrayView<const TSharedRef<IPlugin>> Plugins, const FExportLocalizationOptions& ExportOptions, TFunctionRef<int32(const FString&, FString&)> CommandletExecutor);
+
+/**
+ * Compile UGC localization (if present) for the given plugins, producing LocMeta and LocRes files for consumption by the engine.
+ * 
+ * @param Plugins				The list of plugins to compile.
+ * @param DefaultDescriptor		An optional default UGC localization descriptor to validate any loaded UGC localization descriptors against prior to compiling the localization data.
+ */
+LOCALIZATION_API bool CompileLocalization(TArrayView<const TSharedRef<IPlugin>> Plugins, const FUserGeneratedContentLocalizationDescriptor* DefaultDescriptor = nullptr);
+
+/**
+ * Compile UGC localization (if present) for the given plugin, producing LocMeta and LocRes files for consumption by the engine.
+ *
+ * @param PluginName					The plugin name being compiled.
+ * @param PluginInputContentDirectory	The content directory we'll read the source localization data from when compiling the plugin.
+ * @param PluginOutputContentDirectory	The content directory we'll write the LocMeta and LocRes data to when compiling the plugin.
+ * @param DefaultDescriptor				An optional default UGC localization descriptor to validate any loaded UGC localization descriptors against prior to compiling the localization data.
+ */
+LOCALIZATION_API bool CompileLocalization(const FString& PluginName, const FString& PluginInputContentDirectory, const FString& PluginOutputContentDirectory, const FUserGeneratedContentLocalizationDescriptor* DefaultDescriptor = nullptr);
+
+/**
+ * Load UGC localization source data for the given plugin.
+ * @note This is typically only needed for compilation (which does it internally), but can also be useful if you have other processes that need to read the source data.
+ *
+ * @param PluginName					The plugin name being loaded.
+ * @param PluginInputContentDirectory	The content directory we'll read the source localization data from.
+ * @param OutLocTextHelper				The LocTextHelper to fill with manifest/archive data, re-generated from the source data.
+ * @param DefaultDescriptor				An optional default UGC localization descriptor to validate any loaded UGC localization descriptors against prior to loading the localization data.
+ */
+LOCALIZATION_API ELoadLocalizationResult LoadLocalization(const FString& PluginName, const FString& PluginContentDirectory, TSharedPtr<FLocTextHelper>& OutLocTextHelper, const FUserGeneratedContentLocalizationDescriptor* DefaultDescriptor = nullptr);
+
+}
diff --git a/Engine/Source/Editor/InternationalizationSettings/InternationalizationSettings.Build.cs b/Engine/Source/Editor/InternationalizationSettings/InternationalizationSettings.Build.cs
--- a/Engine/Source/Editor/InternationalizationSettings/InternationalizationSettings.Build.cs
+++ b/Engine/Source/Editor/InternationalizationSettings/InternationalizationSettings.Build.cs
@@ -14,6 +14,7 @@ namespace UnrealBuildTool.Rules
 				    "Engine",
 				    "Slate",
 					"SlateCore",
+					"AppFramework",
 				    "PropertyEditor",
 				    "SharedSettingsWidgets",
                     "Localization",
diff --git a/Engine/Source/Editor/InternationalizationSettings/Private/InternationalizationSettingsModule.cpp b/Engine/Source/Editor/InternationalizationSettings/Private/InternationalizationSettingsModule.cpp
--- a/Engine/Source/Editor/InternationalizationSettings/Private/InternationalizationSettingsModule.cpp
+++ b/Engine/Source/Editor/InternationalizationSettings/Private/InternationalizationSettingsModule.cpp
@@ -4,6 +4,7 @@
 
 #include "Delegates/Delegate.h"
 #include "InternationalizationSettingsModelDetails.h"
+#include "UserGeneratedContentLocalizationDescriptorDetails.h"
 #include "PropertyEditorDelegates.h"
 #include "PropertyEditorModule.h"
 
@@ -21,6 +22,7 @@ void FInternationalizationSettingsModule::StartupModule()
 {
 	FPropertyEditorModule& PropertyModule = FModuleManager::LoadModuleChecked<FPropertyEditorModule>("PropertyEditor");
 	PropertyModule.RegisterCustomClassLayout("InternationalizationSettingsModel", FOnGetDetailCustomizationInstance::CreateStatic(&FInternationalizationSettingsModelDetails::MakeInstance));
+	PropertyModule.RegisterCustomPropertyTypeLayout("UserGeneratedContentLocalizationDescriptor", FOnGetPropertyTypeCustomizationInstance::CreateStatic(&FUserGeneratedContentLocalizationDescriptorDetails::MakeInstance));
 }
 
 void FInternationalizationSettingsModule::ShutdownModule()
@@ -29,5 +31,6 @@ void FInternationalizationSettingsModule::ShutdownModule()
 	{
 		FPropertyEditorModule& PropertyModule = FModuleManager::GetModuleChecked<FPropertyEditorModule>("PropertyEditor");
 		PropertyModule.UnregisterCustomPropertyTypeLayout("InternationalizationSettingsModel");
+		PropertyModule.UnregisterCustomPropertyTypeLayout("UserGeneratedContentLocalizationDescriptor");
 	}
 }
\ No newline at end of file
diff --git a/Engine/Source/Editor/InternationalizationSettings/Private/UserGeneratedContentLocalizationDescriptorDetails.cpp b/Engine/Source/Editor/InternationalizationSettings/Private/UserGeneratedContentLocalizationDescriptorDetails.cpp
new file mode 100644
--- /dev/null
+++ b/Engine/Source/Editor/InternationalizationSettings/Private/UserGeneratedContentLocalizationDescriptorDetails.cpp
@@ -0,0 +1,240 @@
+// Copyright Epic Games, Inc. All Rights Reserved.
+
+#include "UserGeneratedContentLocalizationDescriptorDetails.h"
+#include "UserGeneratedContentLocalization.h"
+
+#include "PropertyHandle.h"
+#include "DetailWidgetRow.h"
+#include "IDetailChildrenBuilder.h"
+
+#include "SCulturePicker.h"
+#include "Widgets/Text/STextBlock.h"
+#include "Widgets/Input/SComboButton.h"
+#include "Widgets/Views/SMultipleOptionTable.h"
+
+#include "Internationalization/Culture.h"
+#include "Internationalization/Internationalization.h"
+
+#define LOCTEXT_NAMESPACE "UserGeneratedContentLocalizationDescriptorDetails"
+
+void FUserGeneratedContentLocalizationDescriptorDetails::CustomizeHeader(TSharedRef<IPropertyHandle> StructPropertyHandle, FDetailWidgetRow& HeaderRow, IPropertyTypeCustomizationUtils& StructCustomizationUtils)
+{
+	if (StructPropertyHandle->HasMetaData("ShowOnlyInnerProperties"))
+	{
+		return;
+	}
+
+	HeaderRow
+		.NameContent()
+		[
+			StructPropertyHandle->CreatePropertyNameWidget()
+		]
+		.ValueContent()
+		[
+			StructPropertyHandle->CreatePropertyValueWidget()
+		];
+}
+
+void FUserGeneratedContentLocalizationDescriptorDetails::CustomizeChildren(TSharedRef<IPropertyHandle> StructPropertyHandle, IDetailChildrenBuilder& StructBuilder, IPropertyTypeCustomizationUtils& StructCustomizationUtils)
+{
+	// Build the list of available cultures used by the child customizations
+	{
+		FUserGeneratedContentLocalizationDescriptor DefaultUGCLocDescriptor;
+		DefaultUGCLocDescriptor.InitializeFromProject();
+
+		const TArray<FCultureRef> LocalizedCultureList = FInternationalization::Get().GetAvailableCultures(DefaultUGCLocDescriptor.CulturesToGenerate, false);
+
+		AvailableCultures.Reset(LocalizedCultureList.Num());
+		Algo::Transform(LocalizedCultureList, AvailableCultures, [](const FCultureRef Culture) -> FCulturePtr { return Culture; });
+
+		// Sort by the current display name so that the order in the "native language" picker and the "cultures to generate" list matches
+		AvailableCultures.StableSort([](FCulturePtr One, FCulturePtr Two)
+		{
+			return One->GetDisplayName() < Two->GetDisplayName();
+		});
+	}
+
+	uint32 NumChildren = 0;
+	StructPropertyHandle->GetNumChildren(NumChildren);
+	for (uint32 ChildIndex = 0; ChildIndex < NumChildren; ++ChildIndex)
+	{
+		if (TSharedPtr<IPropertyHandle> ChildPropertyHandle = StructPropertyHandle->GetChildHandle(ChildIndex))
+		{
+			const FName ChildPropertyName = ChildPropertyHandle->GetProperty()->GetFName();
+			if (ChildPropertyName == GET_MEMBER_NAME_CHECKED(FUserGeneratedContentLocalizationDescriptor, NativeCulture))
+			{
+				CustomizeNativeCulture(ChildPropertyHandle.ToSharedRef(), StructBuilder, StructCustomizationUtils);
+			}
+			else if(ChildPropertyName == GET_MEMBER_NAME_CHECKED(FUserGeneratedContentLocalizationDescriptor, CulturesToGenerate))
+			{
+				CustomizeCulturesToGenerate(ChildPropertyHandle.ToSharedRef(), StructBuilder, StructCustomizationUtils);
+			}
+			else
+			{
+				StructBuilder.AddProperty(ChildPropertyHandle.ToSharedRef());
+			}
+		}
+	}
+}
+
+void FUserGeneratedContentLocalizationDescriptorDetails::CustomizeNativeCulture(TSharedRef<IPropertyHandle> PropertyHandle, IDetailChildrenBuilder& StructBuilder, IPropertyTypeCustomizationUtils& StructCustomizationUtils)
+{
+	NativeCultureHandle = PropertyHandle;
+
+	StructBuilder.AddCustomRow(PropertyHandle->GetPropertyDisplayName())
+		.RowTag(PropertyHandle->GetProperty()->GetFName())
+		.PropertyHandleList({ PropertyHandle })
+		.NameContent()
+		[
+			PropertyHandle->CreatePropertyNameWidget()
+		]
+		.ValueContent()
+		[
+			SNew(SComboButton)
+			.ButtonContent()
+			[
+				SNew(STextBlock)
+				.Font(IPropertyTypeCustomizationUtils::GetRegularFont())
+				.Text(this, &FUserGeneratedContentLocalizationDescriptorDetails::NativeCulture_GetDisplayName)
+			]
+			.MenuContent()
+			[
+				SNew(SBox)
+				.MaxDesiredHeight(400.0f)
+				.MaxDesiredWidth(300.0f)
+				[
+					SNew(SCulturePicker)
+					.InitialSelection(NativeCulture_GetCulture())
+					.OnSelectionChanged(this, &FUserGeneratedContentLocalizationDescriptorDetails::NativeCulture_OnSelectionChanged)
+					.IsCulturePickable(this, &FUserGeneratedContentLocalizationDescriptorDetails::NativeCulture_IsCulturePickable)
+					.DisplayNameFormat(SCulturePicker::ECultureDisplayFormat::ActiveCultureDisplayName)
+					.ViewMode(SCulturePicker::ECulturesViewMode::Flat)
+				]
+			]
+		];
+}
+
+FCulturePtr FUserGeneratedContentLocalizationDescriptorDetails::NativeCulture_GetCulture() const
+{
+	FString NativeCultureName;
+	NativeCultureHandle->GetValue(NativeCultureName);
+
+	return NativeCultureName.IsEmpty()
+		? nullptr
+		: FInternationalization::Get().GetCulture(NativeCultureName);
+}
+
+FText FUserGeneratedContentLocalizationDescriptorDetails::NativeCulture_GetDisplayName() const
+{
+	FCulturePtr NativeCulturePtr = NativeCulture_GetCulture();
+
+	return NativeCulturePtr
+		? FText::AsCultureInvariant(NativeCulturePtr->GetDisplayName())
+		: LOCTEXT("None", "None");
+}
+
+bool FUserGeneratedContentLocalizationDescriptorDetails::NativeCulture_IsCulturePickable(FCulturePtr Culture) const
+{
+	return AvailableCultures.Contains(Culture);
+}
+
+void FUserGeneratedContentLocalizationDescriptorDetails::NativeCulture_OnSelectionChanged(FCulturePtr SelectedCulture, ESelectInfo::Type SelectInfo)
+{
+	NativeCultureHandle->SetValue(SelectedCulture->GetName());
+}
+
+void FUserGeneratedContentLocalizationDescriptorDetails::CustomizeCulturesToGenerate(TSharedRef<IPropertyHandle> PropertyHandle, IDetailChildrenBuilder& StructBuilder, IPropertyTypeCustomizationUtils& StructCustomizationUtils)
+{
+	CulturesToGenerateHandle = PropertyHandle;
+
+	StructBuilder.AddCustomRow(PropertyHandle->GetPropertyDisplayName())
+		.RowTag(PropertyHandle->GetProperty()->GetFName())
+		.PropertyHandleList({ PropertyHandle })
+		.NameContent()
+		[
+			PropertyHandle->CreatePropertyNameWidget()
+		]
+		.ValueContent()
+		.HAlign(HAlign_Fill)
+		.VAlign(VAlign_Fill)
+		[
+			SNew(SMultipleOptionTable<FCulturePtr>, &AvailableCultures)
+			.ListHeight(100.0f)
+			.OnGenerateOptionWidget_Lambda([](FCulturePtr Culture) -> TSharedRef<SWidget>
+			{
+				return SNew(STextBlock)
+					.Font(IPropertyTypeCustomizationUtils::GetRegularFont())
+					.Text(FText::AsCultureInvariant(Culture->GetDisplayName()))
+					.ToolTipText(FText::AsCultureInvariant(Culture->GetName()));
+			})
+			.OnPreBatchSelect(this, &FUserGeneratedContentLocalizationDescriptorDetails::CulturesToGenerate_OnPreBatchSelect)
+			.OnPostBatchSelect(this, &FUserGeneratedContentLocalizationDescriptorDetails::CulturesToGenerate_OnPostBatchSelect)
+			.OnOptionSelectionChanged(this, &FUserGeneratedContentLocalizationDescriptorDetails::CulturesToGenerate_OnCultureSelectionChanged)
+			.IsOptionSelected(this, &FUserGeneratedContentLocalizationDescriptorDetails::CulturesToGenerate_IsCultureSelected)
+		];
+}
+
+void FUserGeneratedContentLocalizationDescriptorDetails::CulturesToGenerate_OnPreBatchSelect()
+{
+	CulturesToGenerate_IsInBatchSelectOperation = true;
+	CulturesToGenerateHandle->NotifyPreChange();
+}
+
+void FUserGeneratedContentLocalizationDescriptorDetails::CulturesToGenerate_OnPostBatchSelect()
+{
+	CulturesToGenerateHandle->NotifyPostChange(EPropertyChangeType::ValueSet);
+	CulturesToGenerate_IsInBatchSelectOperation = false;
+}
+
+void FUserGeneratedContentLocalizationDescriptorDetails::CulturesToGenerate_OnCultureSelectionChanged(bool IsSelected, FCulturePtr Culture)
+{
+	if (!CulturesToGenerate_IsInBatchSelectOperation)
+	{
+		CulturesToGenerateHandle->NotifyPreChange();
+	}
+
+	CulturesToGenerateHandle->EnumerateRawData([IsSelected, &Culture](void* RawData, const int32 /*DataIndex*/, const int32 /*NumDatas*/)
+	{
+		TArray<FString>* RawCulturesToGenerateArray = reinterpret_cast<TArray<FString>*>(RawData);
+		if (IsSelected)
+		{
+			RawCulturesToGenerateArray->AddUnique(Culture->GetName());
+		}
+		else
+		{
+			RawCulturesToGenerateArray->Remove(Culture->GetName());
+		}
+		return true;
+	});
+
+	if (!CulturesToGenerate_IsInBatchSelectOperation)
+	{
+		CulturesToGenerateHandle->NotifyPostChange(IsSelected ? EPropertyChangeType::ArrayAdd : EPropertyChangeType::ArrayRemove);
+	}
+}
+
+bool FUserGeneratedContentLocalizationDescriptorDetails::CulturesToGenerate_IsCultureSelected(FCulturePtr Culture) const
+{
+	FString CultureName = Culture->GetName();
+	TSharedPtr<IPropertyHandleArray> CulturesToGenerateArray = CulturesToGenerateHandle->AsArray();
+	checkf(CulturesToGenerateArray, TEXT("CulturesToGenerate was not an array!"));
+
+	uint32 ElementCount = 0;
+	CulturesToGenerateArray->GetNumElements(ElementCount);
+	for (uint32 Index = 0; Index < ElementCount; ++Index)
+	{
+		const TSharedRef<IPropertyHandle> ElementPropertyHandle = CulturesToGenerateArray->GetElement(Index);
+
+		FString CultureNameAtIndex;
+		ElementPropertyHandle->GetValue(CultureNameAtIndex);
+
+		if (CultureNameAtIndex == CultureName)
+		{
+			return true;
+		}
+	}
+
+	return false;
+}
+
+#undef LOCTEXT_NAMESPACE
diff --git a/Engine/Source/Editor/InternationalizationSettings/Private/UserGeneratedContentLocalizationDescriptorDetails.h b/Engine/Source/Editor/InternationalizationSettings/Private/UserGeneratedContentLocalizationDescriptorDetails.h
new file mode 100644
--- /dev/null
+++ b/Engine/Source/Editor/InternationalizationSettings/Private/UserGeneratedContentLocalizationDescriptorDetails.h
@@ -0,0 +1,46 @@
+// Copyright Epic Games, Inc. All Rights Reserved.
+
+#pragma once
+
+#include "IPropertyTypeCustomization.h"
+#include "Internationalization/CulturePointer.h"
+
+namespace ESelectInfo { enum Type : int; }
+
+class FUserGeneratedContentLocalizationDescriptorDetails : public IPropertyTypeCustomization
+{
+public:
+	/** Makes a new instance of this detail layout class for a specific detail view requesting it */
+	static TSharedRef<IPropertyTypeCustomization> MakeInstance()
+	{
+		return MakeShared<FUserGeneratedContentLocalizationDescriptorDetails>();
+	}
+
+	//~ Begin IPropertyTypeCustomization Interface
+	virtual void CustomizeHeader(TSharedRef<IPropertyHandle> StructPropertyHandle, FDetailWidgetRow& HeaderRow, IPropertyTypeCustomizationUtils& StructCustomizationUtils) override;
+	virtual void CustomizeChildren(TSharedRef<IPropertyHandle> StructPropertyHandle, IDetailChildrenBuilder& StructBuilder, IPropertyTypeCustomizationUtils& StructCustomizationUtils) override;
+	//~ End IPropertyTypeCustomization Interface
+
+private:
+	void CustomizeNativeCulture(TSharedRef<IPropertyHandle> PropertyHandle, IDetailChildrenBuilder& StructBuilder, IPropertyTypeCustomizationUtils& StructCustomizationUtils);
+	FCulturePtr NativeCulture_GetCulture() const;
+	FText NativeCulture_GetDisplayName() const;
+	bool NativeCulture_IsCulturePickable(FCulturePtr Culture) const;
+	void NativeCulture_OnSelectionChanged(FCulturePtr SelectedCulture, ESelectInfo::Type SelectInfo);
+
+	void CustomizeCulturesToGenerate(TSharedRef<IPropertyHandle> PropertyHandle, IDetailChildrenBuilder& StructBuilder, IPropertyTypeCustomizationUtils& StructCustomizationUtils);
+	void CulturesToGenerate_OnPreBatchSelect();
+	void CulturesToGenerate_OnPostBatchSelect();
+	void CulturesToGenerate_OnCultureSelectionChanged(bool IsSelected, FCulturePtr Culture);
+	bool CulturesToGenerate_IsCultureSelected(FCulturePtr Culture) const;
+	void CulturesToGenerate_AddCulture(const FString& CultureName);
+	void CulturesToGenerate_RemoveCulture(const FString& CultureName);
+	
+	TArray<FCulturePtr> AvailableCultures;
+
+	TSharedPtr<IPropertyHandle> NativeCultureHandle;
+
+	TSharedPtr<IPropertyHandle> CulturesToGenerateHandle;
+
+	bool CulturesToGenerate_IsInBatchSelectOperation = false;
+};
diff --git a/Engine/Source/Editor/UnrealEd/Classes/CookOnTheSide/CookOnTheFlyServer.h b/Engine/Source/Editor/UnrealEd/Classes/CookOnTheSide/CookOnTheFlyServer.h
--- a/Engine/Source/Editor/UnrealEd/Classes/CookOnTheSide/CookOnTheFlyServer.h
+++ b/Engine/Source/Editor/UnrealEd/Classes/CookOnTheSide/CookOnTheFlyServer.h
@@ -1001,8 +1001,10 @@ private:
 	/** Write a MapDependencyGraph to a metadata file in the sandbox for the given platform. */
 	void WriteMapDependencyGraph(const ITargetPlatform* TargetPlatform, TMap<FName, TSet<FName>>& MapDependencyGraph);
 
+	void InitializeAllCulturesToCook(TConstArrayView<FString> CookCultures);
+	void CompileDLCLocalization(FBeginCookContext& BeginContext);
 	/** Find localization dependencies for all packages, used to add required localization files as soft references. */
-	void GenerateLocalizationReferences(TConstArrayView<FString> CookCultures);
+	void GenerateLocalizationReferences();
 	void RegisterLocalizationChunkDataGenerator();
 
 	//////////////////////////////////////////////////////////////////////////
diff --git a/Engine/Source/Editor/UnrealEd/Private/CookOnTheFlyServer.cpp b/Engine/Source/Editor/UnrealEd/Private/CookOnTheFlyServer.cpp
--- a/Engine/Source/Editor/UnrealEd/Private/CookOnTheFlyServer.cpp
+++ b/Engine/Source/Editor/UnrealEd/Private/CookOnTheFlyServer.cpp
@@ -136,6 +136,7 @@
 #include "UObject/SavePackage.h"
 #include "UObject/UObjectArray.h"
 #include "UObject/UObjectIterator.h"
+#include "UserGeneratedContentLocalization.h"
 #include "ZenStoreWriter.h"
 
 #define LOCTEXT_NAMESPACE "Cooker"
@@ -9366,6 +9367,7 @@ void UCookOnTheFlyServer::StartCookByTheBook( const FCookByTheBookStartupOptions
 	// Functions in this section are not dependent upon each other and can be ordered arbitrarily or for async performance
 	BeginCookStartShaderCodeLibrary(BeginContext); // start shader code library cooking asynchronously; we block on it later
 	RefreshPlatformAssetRegistries(BeginContext.TargetPlatforms); // Required by BeginCookSandbox stage
+	InitializeAllCulturesToCook(BeginContext.StartupOptions->CookCultures);
 
 	// Clear the sandbox directory, or preserve it and populate iterative cooks
 	// Clear in-memory CookedPackages, or preserve them and cook iteratively in-process
@@ -9382,7 +9384,8 @@ void UCookOnTheFlyServer::StartCookByTheBook( const FCookByTheBookStartupOptions
 	BeginCookEDLCookInfo(BeginContext);
 	BeginCookPackageWriters(BeginContext);
 	GenerateInitialRequests(BeginContext);
-	GenerateLocalizationReferences(BeginContext.StartupOptions->CookCultures);
+	CompileDLCLocalization(BeginContext);
+	GenerateLocalizationReferences();
 	InitializePollables();
 	RecordDLCPackagesFromBaseGame(BeginContext);
 	RegisterCookByTheBookDelegates();
@@ -11091,15 +11094,10 @@ const UE::Cook::FCookSavePackageContext* UCookOnTheFlyServer::FindSaveContext(co
 	return nullptr;
 }
 
-
-void UCookOnTheFlyServer::GenerateLocalizationReferences(TConstArrayView<FString> CookCultures)
+void UCookOnTheFlyServer::InitializeAllCulturesToCook(TConstArrayView<FString> CookCultures)
 {
-	CookByTheBookOptions->SourceToLocalizedPackageVariants.Reset();
 	CookByTheBookOptions->AllCulturesToCook.Reset();
 
-	const UProjectPackagingSettings* const PackagingSettings = GetDefault<UProjectPackagingSettings>();
-
-	// Find all the localized packages and map them back to their source package
 	TArray<FString> AllCulturesToCook(CookCultures);
 	for (const FString& CultureName : CookCultures)
 	{
@@ -11111,7 +11109,41 @@ void UCookOnTheFlyServer::GenerateLocalizationReferences(TConstArrayView<FString
 	}
 	AllCulturesToCook.Sort();
 
-	UE_LOG(LogCook, Display, TEXT("Discovering localized assets for cultures: %s"), *FString::Join(AllCulturesToCook, TEXT(", ")));
+	CookByTheBookOptions->AllCulturesToCook = MoveTemp(AllCulturesToCook);
+}
+
+void UCookOnTheFlyServer::CompileDLCLocalization(FBeginCookContext& BeginContext)
+{
+	if (!IsCookingDLC() || !GetDefault<UUserGeneratedContentLocalizationSettings>()->bCompileDLCLocalizationDuringCook)
+	{
+		return;
+	}
+
+	// Used to validate that we're not loading/compiling invalid cultures during the compile step below
+	FUserGeneratedContentLocalizationDescriptor DefaultUGCLocDescriptor;
+	DefaultUGCLocDescriptor.InitializeFromProject();
+
+	// Also filter the validation list by the cultures we're cooking against
+	DefaultUGCLocDescriptor.CulturesToGenerate.RemoveAll([this](const FString& Culture)
+	{
+		return !CookByTheBookOptions->AllCulturesToCook.Contains(Culture);
+	});
+
+	// Compile UGC localization (if available) for this DLC plugin
+	const FString InputContentDirectory = GetContentDirectoryForDLC();
+	for (const FBeginCookContextPlatform& PlatformContext : BeginContext.PlatformContexts)
+	{
+		const FString OutputContentDirectory = ConvertToFullSandboxPath(InputContentDirectory, /*bForWrite*/true, PlatformContext.TargetPlatform->PlatformName());
+		UserGeneratedContentLocalization::CompileLocalization(CookByTheBookOptions->DlcName, InputContentDirectory, OutputContentDirectory, GetDefault<UUserGeneratedContentLocalizationSettings>()->bValidateDLCLocalizationDuringCook ? &DefaultUGCLocDescriptor : nullptr);
+	}
+}
+
+void UCookOnTheFlyServer::GenerateLocalizationReferences()
+{
+	CookByTheBookOptions->SourceToLocalizedPackageVariants.Reset();
+
+	// Find all the localized packages and map them back to their source package
+	UE_LOG(LogCook, Display, TEXT("Discovering localized assets for cultures: %s"), *FString::Join(CookByTheBookOptions->AllCulturesToCook, TEXT(", ")));
 
 	TArray<FString> RootPaths;
 	FPackageName::QueryRootContentPaths(RootPaths);
@@ -11119,10 +11151,10 @@ void UCookOnTheFlyServer::GenerateLocalizationReferences(TConstArrayView<FString
 	FARFilter Filter;
 	Filter.bRecursivePaths = true;
 	Filter.bIncludeOnlyOnDiskAssets = false;
-	Filter.PackagePaths.Reserve(AllCulturesToCook.Num() * RootPaths.Num());
+	Filter.PackagePaths.Reserve(CookByTheBookOptions->AllCulturesToCook.Num() * RootPaths.Num());
 	for (const FString& RootPath : RootPaths)
 	{
-		for (const FString& CultureName : AllCulturesToCook)
+		for (const FString& CultureName : CookByTheBookOptions->AllCulturesToCook)
 		{
 			FString LocalizedPackagePath = RootPath / TEXT("L10N") / CultureName;
 			Filter.PackagePaths.Add(*LocalizedPackagePath);
@@ -11140,8 +11172,6 @@ void UCookOnTheFlyServer::GenerateLocalizationReferences(TConstArrayView<FString
 		TArray<FName>& LocalizedPackageNames = CookByTheBookOptions->SourceToLocalizedPackageVariants.FindOrAdd(SourcePackageName);
 		LocalizedPackageNames.AddUnique(LocalizedPackageName);
 	}
-
-	CookByTheBookOptions->AllCulturesToCook = MoveTemp(AllCulturesToCook);
 }
 
 void UCookOnTheFlyServer::RegisterLocalizationChunkDataGenerator()
diff --git a/Engine/Source/Runtime/Core/Private/Internationalization/CultureFilter.cpp b/Engine/Source/Runtime/Core/Private/Internationalization/CultureFilter.cpp
new file mode 100644
--- /dev/null
+++ b/Engine/Source/Runtime/Core/Private/Internationalization/CultureFilter.cpp
@@ -0,0 +1,131 @@
+// Copyright Epic Games, Inc. All Rights Reserved.
+
+#include "Internationalization/CultureFilter.h"
+#include "Internationalization/LocalizedTextSourceTypes.h"
+#include "Misc/ConfigCacheIni.h"
+#include "Misc/App.h"
+
+FCultureFilter::FCultureFilter(const TSet<FString>* AvailableCultures)
+{
+	const EBuildConfiguration BuildConfig = FApp::GetBuildConfiguration();
+
+	const ELocalizationLoadFlags TargetFlags 
+		= ELocalizationLoadFlags::Engine
+		| (GIsEditor ? ELocalizationLoadFlags::Editor : ELocalizationLoadFlags::None)
+		| (FApp::IsGame() ? ELocalizationLoadFlags::Game : ELocalizationLoadFlags::None);
+
+	Init(BuildConfig, TargetFlags, AvailableCultures);
+}
+
+FCultureFilter::FCultureFilter(const EBuildConfiguration BuildConfig, const ELocalizationLoadFlags TargetFlags, const TSet<FString>* AvailableCultures)
+{
+	Init(BuildConfig, TargetFlags, AvailableCultures);
+}
+
+void FCultureFilter::Init(const EBuildConfiguration BuildConfig, const ELocalizationLoadFlags TargetFlags, const TSet<FString>* AvailableCultures)
+{
+	// Get the build config string so we can compare it against the config entries
+	FString BuildConfigString;
+	{
+		EBuildConfiguration LocalBuildConfig = BuildConfig;
+		if (LocalBuildConfig == EBuildConfiguration::DebugGame)
+		{
+			// Treat DebugGame and Debug as the same for loc purposes
+			LocalBuildConfig = EBuildConfiguration::Debug;
+		}
+
+		if (LocalBuildConfig != EBuildConfiguration::Unknown)
+		{
+			BuildConfigString = LexToString(LocalBuildConfig);
+		}
+	}
+
+	auto LoadInternationalizationConfigArray = [TargetFlags](const TCHAR* InKey) -> TArray<FString>
+	{
+		check(GConfig && GConfig->IsReadyForUse());
+
+		TArray<FString> FinalArray;
+
+		if (EnumHasAnyFlags(TargetFlags, ELocalizationLoadFlags::Engine))
+		{
+			TArray<FString> EngineArray;
+			GConfig->GetArray(TEXT("Internationalization"), InKey, EngineArray, GEngineIni);
+			FinalArray.Append(MoveTemp(EngineArray));
+		}
+
+		if (EnumHasAnyFlags(TargetFlags, ELocalizationLoadFlags::Editor))
+		{
+			TArray<FString> EditorArray;
+			GConfig->GetArray(TEXT("Internationalization"), InKey, EditorArray, GEditorIni);
+			FinalArray.Append(MoveTemp(EditorArray));
+		}
+
+		if (EnumHasAnyFlags(TargetFlags, ELocalizationLoadFlags::Game))
+		{
+			TArray<FString> GameArray;
+			GConfig->GetArray(TEXT("Internationalization"), InKey, GameArray, GGameIni);
+			FinalArray.Append(MoveTemp(GameArray));
+		}
+
+		return FinalArray;
+	};
+
+	// An array of potentially semicolon separated mapping entries: Culture[;BuildConfig[,BuildConfig,BuildConfig]]
+	// No build config(s) implies all build configs
+	auto ProcessCulturesArray = [&BuildConfigString, AvailableCultures](const TArray<FString>& InCulturesArray, TSet<FString>& OutCulturesSet)
+	{
+		OutCulturesSet.Reserve(InCulturesArray.Num());
+		for (const FString& CultureStr : InCulturesArray)
+		{
+			FString CultureName;
+			FString CultureBuildConfigsStr;
+			if (CultureStr.Split(TEXT(";"), &CultureName, &CultureBuildConfigsStr, ESearchCase::CaseSensitive))
+			{
+				// Check to see if any of the build configs matches our current build config
+				TArray<FString> CultureBuildConfigs;
+				if (CultureBuildConfigsStr.ParseIntoArray(CultureBuildConfigs, TEXT(",")))
+				{
+					bool bIsValidBuildConfig = false;
+					for (const FString& CultureBuildConfig : CultureBuildConfigs)
+					{
+						if (BuildConfigString == CultureBuildConfig)
+						{
+							bIsValidBuildConfig = true;
+							break;
+						}
+					}
+
+					if (!bIsValidBuildConfig)
+					{
+						continue;
+					}
+				}
+			}
+			else
+			{
+				CultureName = CultureStr;
+			}
+
+			if (!AvailableCultures || AvailableCultures->Contains(CultureName))
+			{
+				OutCulturesSet.Add(MoveTemp(CultureName));
+			}
+			else
+			{
+				UE_LOG(LogLocalization, Warning, TEXT("Culture '%s' is unknown and has been ignored when parsing the enabled/disabled culture list."), *CultureName);
+			}
+		}
+		OutCulturesSet.Compact();
+	};
+
+	const TArray<FString> EnabledCulturesArray = LoadInternationalizationConfigArray(TEXT("EnabledCultures"));
+	ProcessCulturesArray(EnabledCulturesArray, EnabledCultures);
+
+	const TArray<FString> DisabledCulturesArray = LoadInternationalizationConfigArray(TEXT("DisabledCultures"));
+	ProcessCulturesArray(DisabledCulturesArray, DisabledCultures);
+}
+
+bool FCultureFilter::IsCultureAllowed(const FString& Culture) const
+{
+	return (EnabledCultures.Num() == 0 || EnabledCultures.Contains(Culture)) && !DisabledCultures.Contains(Culture);
+}
diff --git a/Engine/Source/Runtime/Core/Private/Internationalization/ICUInternationalization.cpp b/Engine/Source/Runtime/Core/Private/Internationalization/ICUInternationalization.cpp
--- a/Engine/Source/Runtime/Core/Private/Internationalization/ICUInternationalization.cpp
+++ b/Engine/Source/Runtime/Core/Private/Internationalization/ICUInternationalization.cpp
@@ -5,6 +5,7 @@
 #include "Misc/ScopeLock.h"
 #include "Misc/Paths.h"
 #include "Internationalization/Culture.h"
+#include "Internationalization/CultureFilter.h"
 #include "Internationalization/Internationalization.h"
 #include "Internationalization/Cultures/LeetCulture.h"
 #include "Internationalization/Cultures/KeysCulture.h"
@@ -12,6 +13,7 @@
 #include "Misc/CoreStats.h"
 #include "Misc/ConfigCacheIni.h"
 #include "Misc/App.h"
+#include "Algo/Transform.h"
 
 #if UE_ENABLE_ICU
 
@@ -459,75 +461,11 @@ void FICUInternationalization::ConditionalInitializeAllowedCultures()
 
 	bHasInitializedAllowedCultures = true;
 
-	// Get our current build config string so we can compare it against the config entries
-	FString BuildConfigString;
-	{
-		EBuildConfiguration BuildConfig = FApp::GetBuildConfiguration();
-		if (BuildConfig == EBuildConfiguration::DebugGame)
-		{
-			// Treat DebugGame and Debug as the same for loc purposes
-			BuildConfig = EBuildConfiguration::Debug;
-		}
-
-		if (BuildConfig != EBuildConfiguration::Unknown)
-		{
-			BuildConfigString = LexToString(BuildConfig);
-		}
-	}
-
-	// An array of potentially semicolon separated mapping entries: Culture[;BuildConfig[,BuildConfig,BuildConfig]]
-	// No build config(s) implies all build configs
-	auto ProcessCulturesArray = [this, &BuildConfigString](const TArray<FString>& InCulturesArray, TSet<FString>& OutCulturesSet)
-	{
-		OutCulturesSet.Reserve(InCulturesArray.Num());
-		for (const FString& CultureStr : InCulturesArray)
-		{
-			FString CultureName;
-			FString CultureBuildConfigsStr;
-			if (CultureStr.Split(TEXT(";"), &CultureName, &CultureBuildConfigsStr, ESearchCase::CaseSensitive))
-			{
-				// Check to see if any of the build configs matches our current build config
-				TArray<FString> CultureBuildConfigs;
-				if (CultureBuildConfigsStr.ParseIntoArray(CultureBuildConfigs, TEXT(",")))
-				{
-					bool bIsValidBuildConfig = false;
-					for (const FString& CultureBuildConfig : CultureBuildConfigs)
-					{
-						if (BuildConfigString == CultureBuildConfig)
-						{
-							bIsValidBuildConfig = true;
-							break;
-						}
-					}
-
-					if (!bIsValidBuildConfig)
-					{
-						continue;
-					}
-				}
-			}
-			else
-			{
-				CultureName = CultureStr;
-			}
-
-			if (AllAvailableCulturesMap.Contains(CultureName))
-			{
-				OutCulturesSet.Add(MoveTemp(CultureName));
-			}
-			else
-			{
-				UE_LOG(LogICUInternationalization, Warning, TEXT("Culture '%s' is unknown and has been ignored when parsing the enabled/disabled culture list."), *CultureName);
-			}
-		}
-		OutCulturesSet.Compact();
-	};
-
-	const TArray<FString> EnabledCulturesArray = LoadInternationalizationConfigArray(TEXT("EnabledCultures"));
-	ProcessCulturesArray(EnabledCulturesArray, EnabledCultures);
+	TSet<FString> AvailableCulturesSet;
+	AvailableCulturesSet.Reserve(AllAvailableCulturesMap.Num());
+	Algo::Transform(AllAvailableCulturesMap, AvailableCulturesSet, &TTuple<FString, int32>::Key);
 
-	const TArray<FString> DisabledCulturesArray = LoadInternationalizationConfigArray(TEXT("DisabledCultures"));
-	ProcessCulturesArray(DisabledCulturesArray, DisabledCultures);
+	AllowedCulturesFilter = MakePimpl<FCultureFilter>(&AvailableCulturesSet);
 }
 
 bool FICUInternationalization::IsCultureRemapped(const FString& Name, FString* OutMappedCulture)
@@ -550,7 +488,7 @@ bool FICUInternationalization::IsCultureAllowed(const FString& Name)
 	// Make sure we've loaded the allowed cultures lists (the config system may not have been available when we were first initialized)
 	ConditionalInitializeAllowedCultures();
 
-	return (EnabledCultures.Num() == 0 || EnabledCultures.Contains(Name)) && !DisabledCultures.Contains(Name);
+	return AllowedCulturesFilter->IsCultureAllowed(Name);
 }
 
 void FICUInternationalization::RefreshCultureDisplayNames(const TArray<FString>& InPrioritizedDisplayCultureNames)
@@ -570,8 +508,7 @@ void FICUInternationalization::RefreshCachedConfigData()
 	ConditionalInitializeCultureMappings();
 
 	bHasInitializedAllowedCultures = false;
-	EnabledCultures.Reset();
-	DisabledCultures.Reset();
+	AllowedCulturesFilter.Reset();
 	ConditionalInitializeAllowedCultures();
 }
 
diff --git a/Engine/Source/Runtime/Core/Private/Internationalization/ICUInternationalization.h b/Engine/Source/Runtime/Core/Private/Internationalization/ICUInternationalization.h
--- a/Engine/Source/Runtime/Core/Private/Internationalization/ICUInternationalization.h
+++ b/Engine/Source/Runtime/Core/Private/Internationalization/ICUInternationalization.h
@@ -8,7 +8,9 @@
 #include "Containers/Map.h"
 #include "Internationalization/CulturePointer.h"
 #include "Templates/UniquePtr.h"
+#include "Templates/PimplPtr.h"
 
+class FCultureFilter;
 class FInternationalization;
 
 #if UE_ENABLE_ICU
@@ -98,8 +100,7 @@ private:
 	TMap<FString, FString> CultureMappings;
 
 	bool bHasInitializedAllowedCultures;
-	TSet<FString> EnabledCultures;
-	TSet<FString> DisabledCultures;
+	TPimplPtr<FCultureFilter> AllowedCulturesFilter;
 
 	TMap<FString, FCultureRef> CachedCultures;
 	FCriticalSection CachedCulturesCS;
diff --git a/Engine/Source/Runtime/Core/Public/Internationalization/CultureFilter.h b/Engine/Source/Runtime/Core/Public/Internationalization/CultureFilter.h
new file mode 100644
--- /dev/null
+++ b/Engine/Source/Runtime/Core/Public/Internationalization/CultureFilter.h
@@ -0,0 +1,28 @@
+// Copyright Epic Games, Inc. All Rights Reserved.
+#pragma once
+
+#include "Containers/Set.h"
+#include "Containers/UnrealString.h"
+
+enum class EBuildConfiguration : uint8;
+enum class ELocalizationLoadFlags : uint8;
+
+/** Cache of enabled/disabled cultures loaded from config data */
+class CORE_API FCultureFilter
+{
+public:
+	/** Construct a culture filter based on the current build config and target, optionally filtering any loaded cultures against the set of available cultures */
+	explicit FCultureFilter(const TSet<FString>* AvailableCultures = nullptr);
+
+	/** Construct a culture filter based on the given build config and target, optionally filtering any loaded cultures against the set of available cultures */
+	FCultureFilter(const EBuildConfiguration BuildConfig, const ELocalizationLoadFlags TargetFlags, const TSet<FString>* AvailableCultures = nullptr);
+
+	/** Does the given culture pass this filter? */
+	bool IsCultureAllowed(const FString& Culture) const;
+
+private:
+	void Init(const EBuildConfiguration BuildConfig, const ELocalizationLoadFlags TargetFlags, const TSet<FString>* AvailableCultures);
+
+	TSet<FString> EnabledCultures;
+	TSet<FString> DisabledCultures;
+};
