Default to exporting localization files with the copyright notice set on the project

This can also be overridden using the "CopyrightNotice" option in the "CommonSettings" of a GatherText config

diff --git a/Engine/Source/Developer/Localization/Private/LocTextHelper.cpp b/Engine/Source/Developer/Localization/Private/LocTextHelper.cpp
--- a/Engine/Source/Developer/Localization/Private/LocTextHelper.cpp
+++ b/Engine/Source/Developer/Localization/Private/LocTextHelper.cpp
@@ -524,6 +524,16 @@ TSharedPtr<ILocFileNotifies> FLocTextHelper::GetLocFileNotifies() const
 	return LocFileNotifies;
 }
 
+void FLocTextHelper::SetCopyrightNotice(const FString& InCopyrightNotice)
+{
+	CopyrightNotice = InCopyrightNotice;
+}
+
+const FString& FLocTextHelper::GetCopyrightNotice() const
+{
+	return CopyrightNotice;
+}
+
 const FString& FLocTextHelper::GetNativeCulture() const
 {
 	return NativeCulture;
diff --git a/Engine/Source/Developer/Localization/Private/PortableObjectFormatDOM.cpp b/Engine/Source/Developer/Localization/Private/PortableObjectFormatDOM.cpp
--- a/Engine/Source/Developer/Localization/Private/PortableObjectFormatDOM.cpp
+++ b/Engine/Source/Developer/Localization/Private/PortableObjectFormatDOM.cpp
@@ -753,23 +753,23 @@ bool FPortableObjectFormatDOM::FromString( const FString& InStr, FText* OutError
 	return true;
 }
 
-void FPortableObjectFormatDOM::CreateNewHeader()
+void FPortableObjectFormatDOM::CreateNewHeader(  const FString& InCopyrightNotice )
 {
 	// Reference: http://www.gnu.org/software/gettext/manual/gettext.html#Header-Entry
 	// Reference: http://www.gnu.org/software/gettext/manual/html_node/Header-Entry.html
 
 	//Hard code some header entries for now in the following format
 	/*
-	# Engine English translation
-	# Copyright Epic Games, Inc. All Rights Reserved.
+	# {ProjectName} {LanguageName} translation
+	# {CopyrightNotice}
 	#
 	msgid ""
 	msgstr ""
-	"Project-Id-Version: Engine\n"
+	"Project-Id-Version: {ProjectName}\n"
 	"POT-Creation-Date: 2014-1-31 04:16+0000\n"
 	"PO-Revision-Date: 2014-1-31 04:16+0000\n"
 	"Language-Team: \n"
-	"Language: en-us\n"
+	"Language: {LanguageCode}\n"
 	"MIME-Version: 1.0\n"
 	"Content-Type: text/plain; charset=UTF-8\n"
 	"Content-Transfer-Encoding: 8bit\n"
@@ -790,7 +790,10 @@ void FPortableObjectFormatDOM::CreateNewHeader()
 	Header.SetEntryValue( TEXT("Plural-Forms"), Language.GetPluralForms() );
 
 	Header.Comments.Add( FString::Printf(TEXT("%s %s translation."), *GetProjectName(), *Language.EnglishName() ) );
-	Header.Comments.Add( TEXT("Copyright Epic Games, Inc. All Rights Reserved.") );
+	if (!InCopyrightNotice.IsEmpty())
+	{
+		Header.Comments.Add(InCopyrightNotice);
+	}
 	Header.Comments.Add( FString(TEXT("")) );
 }
 
diff --git a/Engine/Source/Developer/Localization/Private/PortableObjectPipeline.cpp b/Engine/Source/Developer/Localization/Private/PortableObjectPipeline.cpp
--- a/Engine/Source/Developer/Localization/Private/PortableObjectPipeline.cpp
+++ b/Engine/Source/Developer/Localization/Private/PortableObjectPipeline.cpp
@@ -340,7 +340,7 @@ namespace
 		}
 
 		NewPortableObject.SetProjectName(FPaths::GetBaseFilename(InPOFilePath));
-		NewPortableObject.CreateNewHeader();
+		NewPortableObject.CreateNewHeader(InLocTextHelper.GetCopyrightNotice());
 		UpdatePOFileHeaderForSettings(NewPortableObject, InTextCollapseMode, InPOFormat);
 
 		// Add each manifest entry to the PO file
diff --git a/Engine/Source/Developer/Localization/Private/UserGeneratedContentLocalization.cpp b/Engine/Source/Developer/Localization/Private/UserGeneratedContentLocalization.cpp
--- a/Engine/Source/Developer/Localization/Private/UserGeneratedContentLocalization.cpp
+++ b/Engine/Source/Developer/Localization/Private/UserGeneratedContentLocalization.cpp
@@ -232,6 +232,8 @@ bool ExportLocalization(TArrayView<const TSharedRef<IPlugin>> Plugins, const FEx
 				ConfigSection.Add(TEXT("ArchiveName"), FString::Printf(TEXT("%s.archive"), *Plugin->GetName()));
 				ConfigSection.Add(TEXT("PortableObjectName"), FString::Printf(TEXT("%s.po"), *Plugin->GetName()));
 
+				ConfigSection.Add(TEXT("CopyrightNotice"), ExportOptions.CopyrightNotice);
+
 				ConfigSection.Add(TEXT("NativeCulture"), ExportOptions.UGCLocDescriptor.NativeCulture);
 				for (const FString& CultureToGenerate : CulturesToGenerate)
 				{
diff --git a/Engine/Source/Developer/Localization/Public/LocTextHelper.h b/Engine/Source/Developer/Localization/Public/LocTextHelper.h
--- a/Engine/Source/Developer/Localization/Public/LocTextHelper.h
+++ b/Engine/Source/Developer/Localization/Public/LocTextHelper.h
@@ -320,6 +320,16 @@ public:
 	 */
 	TSharedPtr<ILocFileNotifies> GetLocFileNotifies() const;
 
+	/**
+	 * Set the optional copyright notice to insert into files that support it.
+	 */
+	void SetCopyrightNotice(const FString& InCopyrightNotice);
+
+	/**
+	 * @return Get the optional copyright notice to insert into files that support it.
+	 */
+	const FString& GetCopyrightNotice() const;
+
 	/**
 	 * @return Get the culture code of the native culture (eg, en), or an empty string if the native culture is unknown.
 	 */
@@ -965,6 +975,9 @@ private:
 	/** Interface for allowing source control integration (may be null) */
 	TSharedPtr<ILocFileNotifies> LocFileNotifies;
 
+	/** Optional copyright notice to insert into files that support it */
+	FString CopyrightNotice;
+
 	/** Loaded manifest */
 	TSharedPtr<FInternationalizationManifest> Manifest;
 
diff --git a/Engine/Source/Developer/Localization/Public/PortableObjectFormatDOM.h b/Engine/Source/Developer/Localization/Public/PortableObjectFormatDOM.h
--- a/Engine/Source/Developer/Localization/Public/PortableObjectFormatDOM.h
+++ b/Engine/Source/Developer/Localization/Public/PortableObjectFormatDOM.h
@@ -346,7 +346,7 @@ public:
 	bool FromString( const FString& InStr, FText* OutErrorMsg = nullptr );
 
 	/** Creates a header entry based on the project and language info. */
-	void CreateNewHeader();
+	void CreateNewHeader( const FString& InCopyrightNotice = FString() );
 
 	/**
 	 * Sets a header entry value.
diff --git a/Engine/Source/Developer/Localization/Public/UserGeneratedContentLocalization.h b/Engine/Source/Developer/Localization/Public/UserGeneratedContentLocalization.h
--- a/Engine/Source/Developer/Localization/Public/UserGeneratedContentLocalization.h
+++ b/Engine/Source/Developer/Localization/Public/UserGeneratedContentLocalization.h
@@ -140,6 +140,9 @@ struct FExportLocalizationOptions
 
 	/** The category to use for the exported localization target (only used when bUpdatePluginDescriptor is true) */
 	ELocalizedTextSourceCategory LocalizationCategory = ELocalizedTextSourceCategory::Game;
+
+	/** An optional copyright notice to insert into the exported files */
+	FString CopyrightNotice;
 };
 
 enum class ELoadLocalizationResult : uint8
diff --git a/Engine/Source/Editor/TranslationEditor/Private/TranslationDataManager.cpp b/Engine/Source/Editor/TranslationEditor/Private/TranslationDataManager.cpp
--- a/Engine/Source/Editor/TranslationEditor/Private/TranslationDataManager.cpp
+++ b/Engine/Source/Editor/TranslationEditor/Private/TranslationDataManager.cpp
@@ -29,6 +29,7 @@
 #include "LocalizationConfigurationScript.h"
 #include "Serialization/JsonInternationalizationArchiveSerializer.h"
 #include "Serialization/JsonInternationalizationManifestSerializer.h"
+#include "GeneralProjectSettings.h"
 
 
 DEFINE_LOG_CATEGORY_STATIC(LogTranslationEditor, Log, All);
@@ -890,7 +891,7 @@ bool FTranslationDataManager::SaveSelectedTranslations(TArray<UTranslationUnit*>
 				FPortableObjectFormatDOM PortableObjectDom;
 				PortableObjectDom.SetProjectName(ManifestAndArchiveName);
 				PortableObjectDom.SetLanguage(CultureName);
-				PortableObjectDom.CreateNewHeader();
+				PortableObjectDom.CreateNewHeader(GetDefault<UGeneralProjectSettings>()->CopyrightNotice);
 				PortableObjectPipeline::UpdatePOFileHeaderForSettings(PortableObjectDom, LocalizationTarget->Settings.ExportSettings.CollapseMode, LocalizationTarget->Settings.ExportSettings.POFormat);
 
 				TArray<UTranslationUnit*>& TranslationsArray = DataManager->GetAllTranslationsArray();
diff --git a/Engine/Source/Editor/TranslationEditor/TranslationEditor.Build.cs b/Engine/Source/Editor/TranslationEditor/TranslationEditor.Build.cs
--- a/Engine/Source/Editor/TranslationEditor/TranslationEditor.Build.cs
+++ b/Engine/Source/Editor/TranslationEditor/TranslationEditor.Build.cs
@@ -29,6 +29,7 @@ public class TranslationEditor : ModuleRules
 				"Core",
 				"CoreUObject",
 				"Engine",
+				"EngineSettings",
                 "InputCore",
 				"Json",
                 "PropertyEditor",
diff --git a/Engine/Source/Editor/UnrealEd/Private/Commandlets/ExportDialogueScriptCommandlet.cpp b/Engine/Source/Editor/UnrealEd/Private/Commandlets/ExportDialogueScriptCommandlet.cpp
--- a/Engine/Source/Editor/UnrealEd/Private/Commandlets/ExportDialogueScriptCommandlet.cpp
+++ b/Engine/Source/Editor/UnrealEd/Private/Commandlets/ExportDialogueScriptCommandlet.cpp
@@ -161,7 +161,8 @@ int32 UExportDialogueScriptCommandlet::Main(const FString& Params)
 	}
 
 	// Load the manifest and all archives
-	FLocTextHelper LocTextHelper(SourcePath, ManifestName, ArchiveName, NativeCulture, CulturesToGenerate, MakeShareable(new FLocFileSCCNotifies(SourceControlInfo)));
+	FLocTextHelper LocTextHelper(SourcePath, ManifestName, ArchiveName, NativeCulture, CulturesToGenerate, GatherManifestHelper->GetLocFileNotifies(), GatherManifestHelper->GetPlatformSplitMode());
+	LocTextHelper.SetCopyrightNotice(GatherManifestHelper->GetCopyrightNotice());
 	{
 		FText LoadError;
 		if (!LocTextHelper.LoadAll(ELocTextHelperLoadFlags::LoadOrCreate, &LoadError))
diff --git a/Engine/Source/Editor/UnrealEd/Private/Commandlets/GatherTextCommandlet.cpp b/Engine/Source/Editor/UnrealEd/Private/Commandlets/GatherTextCommandlet.cpp
--- a/Engine/Source/Editor/UnrealEd/Private/Commandlets/GatherTextCommandlet.cpp
+++ b/Engine/Source/Editor/UnrealEd/Private/Commandlets/GatherTextCommandlet.cpp
@@ -9,6 +9,7 @@
 #include "UObject/Package.h"
 #include "UObject/GCObjectScopeGuard.h"
 #include "SourceControlHelpers.h"
+#include "GeneralProjectSettings.h"
 
 DEFINE_LOG_CATEGORY_STATIC(LogGatherTextCommandlet, Log, All);
 
@@ -168,8 +169,15 @@ int32 UGatherTextCommandlet::ProcessGatherConfig(const FString& GatherTextConfig
 		LocalizationTargetName = FPaths::GetBaseFilename(ManifestName);
 	}
 
+	FString CopyrightNotice;
+	if (!GetStringFromConfig(TEXT("CommonSettings"), TEXT("CopyrightNotice"), CopyrightNotice, GatherTextConfigPath))
+	{
+		CopyrightNotice = GetDefault<UGeneralProjectSettings>()->CopyrightNotice;
+	}
+
 	// Basic helper that can be used only to gather a new manifest for writing
 	TSharedRef<FLocTextHelper> CommandletGatherManifestHelper = MakeShared<FLocTextHelper>(LocalizationTargetName, MakeShared<FLocFileSCCNotifies>(CommandletSourceControlInfo), PlatformSplitMode);
+	CommandletGatherManifestHelper->SetCopyrightNotice(CopyrightNotice);
 	CommandletGatherManifestHelper->LoadManifest(ELocTextHelperLoadFlags::Create);
 
 	const FString GatherTextStepPrefix = TEXT("GatherTextStep");
diff --git a/Engine/Source/Editor/UnrealEd/Private/Commandlets/GenerateGatherArchiveCommandlet.cpp b/Engine/Source/Editor/UnrealEd/Private/Commandlets/GenerateGatherArchiveCommandlet.cpp
--- a/Engine/Source/Editor/UnrealEd/Private/Commandlets/GenerateGatherArchiveCommandlet.cpp
+++ b/Engine/Source/Editor/UnrealEd/Private/Commandlets/GenerateGatherArchiveCommandlet.cpp
@@ -219,6 +219,7 @@ int32 UGenerateGatherArchiveCommandlet::Main( const FString& Params )
 
 	// Load the manifest and all archives
 	FLocTextHelper LocTextHelper(DestinationPath, ManifestName, ArchiveName, NativeCulture, CulturesToGenerate, GatherManifestHelper->GetLocFileNotifies(), GatherManifestHelper->GetPlatformSplitMode());
+	LocTextHelper.SetCopyrightNotice(GatherManifestHelper->GetCopyrightNotice());
 	{
 		FText LoadError;
 		if (!LocTextHelper.LoadAll(ELocTextHelperLoadFlags::LoadOrCreate, &LoadError))
diff --git a/Engine/Source/Editor/UnrealEd/Private/Commandlets/GenerateTextLocalizationReportCommandlet.cpp b/Engine/Source/Editor/UnrealEd/Private/Commandlets/GenerateTextLocalizationReportCommandlet.cpp
--- a/Engine/Source/Editor/UnrealEd/Private/Commandlets/GenerateTextLocalizationReportCommandlet.cpp
+++ b/Engine/Source/Editor/UnrealEd/Private/Commandlets/GenerateTextLocalizationReportCommandlet.cpp
@@ -157,7 +157,8 @@ bool UGenerateTextLocalizationReportCommandlet::ProcessWordCountReport(const FSt
 	}
 
 	// Load the manifest and all archives
-	FLocTextHelper LocTextHelper(SourcePath, ManifestName, ArchiveName, FString(), CulturesToGenerate, MakeShareable(new FLocFileSCCNotifies(SourceControlInfo)));
+	FLocTextHelper LocTextHelper(SourcePath, ManifestName, ArchiveName, FString(), CulturesToGenerate, GatherManifestHelper->GetLocFileNotifies(), GatherManifestHelper->GetPlatformSplitMode());
+	LocTextHelper.SetCopyrightNotice(GatherManifestHelper->GetCopyrightNotice());
 	{
 		FText LoadError;
 		if (!LocTextHelper.LoadAll(ELocTextHelperLoadFlags::LoadOrCreate, &LoadError))
diff --git a/Engine/Source/Editor/UnrealEd/Private/Commandlets/GenerateTextLocalizationResourceCommandlet.cpp b/Engine/Source/Editor/UnrealEd/Private/Commandlets/GenerateTextLocalizationResourceCommandlet.cpp
--- a/Engine/Source/Editor/UnrealEd/Private/Commandlets/GenerateTextLocalizationResourceCommandlet.cpp
+++ b/Engine/Source/Editor/UnrealEd/Private/Commandlets/GenerateTextLocalizationResourceCommandlet.cpp
@@ -155,6 +155,7 @@ int32 UGenerateTextLocalizationResourceCommandlet::Main(const FString& Params)
 
 	// Load the manifest and all archives
 	FLocTextHelper LocTextHelper(SourcePath, ManifestName, ArchiveName, NativeCultureName, CulturesToGenerate, GatherManifestHelper->GetLocFileNotifies(), GatherManifestHelper->GetPlatformSplitMode());
+	LocTextHelper.SetCopyrightNotice(GatherManifestHelper->GetCopyrightNotice());
 	{
 		FText LoadError;
 		if (!LocTextHelper.LoadAll(ELocTextHelperLoadFlags::LoadOrCreate, &LoadError))
diff --git a/Engine/Source/Editor/UnrealEd/Private/Commandlets/ImportDialogueScriptCommandlet.cpp b/Engine/Source/Editor/UnrealEd/Private/Commandlets/ImportDialogueScriptCommandlet.cpp
--- a/Engine/Source/Editor/UnrealEd/Private/Commandlets/ImportDialogueScriptCommandlet.cpp
+++ b/Engine/Source/Editor/UnrealEd/Private/Commandlets/ImportDialogueScriptCommandlet.cpp
@@ -134,7 +134,8 @@ int32 UImportDialogueScriptCommandlet::Main(const FString& Params)
 	}
 
 	// Load the manifest and all archives
-	FLocTextHelper LocTextHelper(DestinationPath, ManifestName, ArchiveName, NativeCulture, CulturesToGenerate, MakeShareable(new FLocFileSCCNotifies(SourceControlInfo)));
+	FLocTextHelper LocTextHelper(DestinationPath, ManifestName, ArchiveName, NativeCulture, CulturesToGenerate, GatherManifestHelper->GetLocFileNotifies(), GatherManifestHelper->GetPlatformSplitMode());
+	LocTextHelper.SetCopyrightNotice(GatherManifestHelper->GetCopyrightNotice());
 	{
 		FText LoadError;
 		if (!LocTextHelper.LoadAll(ELocTextHelperLoadFlags::LoadOrCreate, &LoadError))
diff --git a/Engine/Source/Editor/UnrealEd/Private/Commandlets/ImportLocalizedDialogueCommandlet.cpp b/Engine/Source/Editor/UnrealEd/Private/Commandlets/ImportLocalizedDialogueCommandlet.cpp
--- a/Engine/Source/Editor/UnrealEd/Private/Commandlets/ImportLocalizedDialogueCommandlet.cpp
+++ b/Engine/Source/Editor/UnrealEd/Private/Commandlets/ImportLocalizedDialogueCommandlet.cpp
@@ -136,7 +136,8 @@ int32 UImportLocalizedDialogueCommandlet::Main(const FString& Params)
 	}
 
 	// Load the manifest and all archives
-	FLocTextHelper LocTextHelper(SourcePath, ManifestName, ArchiveName, NativeCulture, CulturesToGenerate, MakeShareable(new FLocFileSCCNotifies(SourceControlInfo)));
+	FLocTextHelper LocTextHelper(SourcePath, ManifestName, ArchiveName, NativeCulture, CulturesToGenerate, GatherManifestHelper->GetLocFileNotifies(), GatherManifestHelper->GetPlatformSplitMode());
+	LocTextHelper.SetCopyrightNotice(GatherManifestHelper->GetCopyrightNotice());
 	{
 		FText LoadError;
 		if (!LocTextHelper.LoadAll(ELocTextHelperLoadFlags::LoadOrCreate, &LoadError))
diff --git a/Engine/Source/Editor/UnrealEd/Private/Commandlets/InternationalizationExportCommandlet.cpp b/Engine/Source/Editor/UnrealEd/Private/Commandlets/InternationalizationExportCommandlet.cpp
--- a/Engine/Source/Editor/UnrealEd/Private/Commandlets/InternationalizationExportCommandlet.cpp
+++ b/Engine/Source/Editor/UnrealEd/Private/Commandlets/InternationalizationExportCommandlet.cpp
@@ -161,6 +161,7 @@ int32 UInternationalizationExportCommandlet::Main(const FString& Params)
 	{
 		// Load the manifest and all archives
 		FLocTextHelper LocTextHelper(DestinationPath, ManifestName, ArchiveName, NativeCultureName, CulturesToGenerate, GatherManifestHelper->GetLocFileNotifies(), GatherManifestHelper->GetPlatformSplitMode());
+		LocTextHelper.SetCopyrightNotice(GatherManifestHelper->GetCopyrightNotice());
 		{
 			FText LoadError;
 			if (!LocTextHelper.LoadAll(ELocTextHelperLoadFlags::LoadOrCreate, &LoadError))
@@ -185,6 +186,7 @@ int32 UInternationalizationExportCommandlet::Main(const FString& Params)
 
 		// Load the manifest and all archives
 		FLocTextHelper LocTextHelper(SourcePath, ManifestName, ArchiveName, NativeCultureName, CulturesToGenerate, GatherManifestHelper->GetLocFileNotifies(), GatherManifestHelper->GetPlatformSplitMode());
+		LocTextHelper.SetCopyrightNotice(GatherManifestHelper->GetCopyrightNotice());
 		{
 			FText LoadError;
 			if (!LocTextHelper.LoadAll(ELocTextHelperLoadFlags::LoadOrCreate, &LoadError))
