Exposed field filters to the Localization Dashboard

diff --git a/Engine/Source/Developer/Localization/Private/LocalizationConfigurationScript.cpp b/Engine/Source/Developer/Localization/Private/LocalizationConfigurationScript.cpp
--- a/Engine/Source/Developer/Localization/Private/LocalizationConfigurationScript.cpp
+++ b/Engine/Source/Developer/Localization/Private/LocalizationConfigurationScript.cpp
@@ -419,7 +419,7 @@ namespace LocalizationConfigurationScript
 				ConfigSection.Add( TEXT("ExcludePathFilters"), FString::Printf(TEXT("%s%s"), *FLocalizationGatherPathRootUtil::GetResolvedPathRootToken(ExcludePath.PathRoot), *ExcludePath.Pattern) );
 			}
 
-			// Package Extensions
+			// Key Specifications
 			for (const FMetaDataKeyGatherSpecification& Specification : Target->Settings.GatherFromMetaData.KeySpecifications)
 			{
 				ConfigSection.Add( TEXT("InputKeys"), Specification.MetaDataKey.Name );
@@ -427,6 +427,26 @@ namespace LocalizationConfigurationScript
 				ConfigSection.Add( TEXT("OutputKeys"), Specification.TextKeyPattern.Pattern );
 			}
 
+			// Field Type Filters
+			for (const FString& FieldTypeToInclude : Target->Settings.GatherFromMetaData.FieldTypesToInclude)
+			{
+				ConfigSection.Add(TEXT("FieldTypesToInclude"), FieldTypeToInclude);
+			}
+			for (const FString& FieldTypeToExclude : Target->Settings.GatherFromMetaData.FieldTypesToExclude)
+			{
+				ConfigSection.Add(TEXT("FieldTypesToExclude"), FieldTypeToExclude);
+			}
+
+			// Field Owner Type Filters
+			for (const FString& FieldOwnerTypeToInclude : Target->Settings.GatherFromMetaData.FieldOwnerTypesToInclude)
+			{
+				ConfigSection.Add(TEXT("FieldOwnerTypesToInclude"), FieldOwnerTypeToInclude);
+			}
+			for (const FString& FieldOwnerTypeToExclude : Target->Settings.GatherFromMetaData.FieldOwnerTypesToExclude)
+			{
+				ConfigSection.Add(TEXT("FieldOwnerTypesToExclude"), FieldOwnerTypeToExclude);
+			}
+
 			ConfigSection.Add( TEXT("ShouldGatherFromEditorOnlyData"), Target->Settings.GatherFromMetaData.ShouldGatherFromEditorOnlyData ? TEXT("true") : TEXT("false") );
 		}
 
diff --git a/Engine/Source/Developer/Localization/Public/LocalizationTargetTypes.h b/Engine/Source/Developer/Localization/Public/LocalizationTargetTypes.h
--- a/Engine/Source/Developer/Localization/Public/LocalizationTargetTypes.h
+++ b/Engine/Source/Developer/Localization/Public/LocalizationTargetTypes.h
@@ -281,6 +281,22 @@ struct FGatherTextFromMetaDataConfiguration
 	UPROPERTY(config, EditAnywhere, Category = "MetaData")
 	TArray<FMetaDataKeyGatherSpecification> KeySpecifications;
 
+	/** List of field types (eg, Property, Function, ScriptStruct, Enum, etc) that should be included in the gather, or empty to include everything. */
+	UPROPERTY(config, EditAnywhere, Category = "MetaData")
+	TArray<FString> FieldTypesToInclude;
+
+	/** List of field types (eg, Property, Function, ScriptStruct, Enum, etc) the should be excluded from the gather. */
+	UPROPERTY(config, EditAnywhere, Category = "MetaData")
+	TArray<FString> FieldTypesToExclude;
+
+	/** List of field owner types (eg, MyClass, MyStruct, etc) that should have fields within them included in the gather, or empty to include everything. */
+	UPROPERTY(config, EditAnywhere, Category = "MetaData")
+	TArray<FString> FieldOwnerTypesToInclude;
+
+	/** List of field owner types (eg, MyClass, MyStruct, etc) that should have fields within them excluded from the gather. */
+	UPROPERTY(config, EditAnywhere, Category = "MetaData")
+	TArray<FString> FieldOwnerTypesToExclude;
+
 	/* If enabled, data that is specified as editor-only may be processed for gathering. */
 	UPROPERTY(config, EditAnywhere, Category = "Filter")
 	bool ShouldGatherFromEditorOnlyData;
