Added extra localization logging

#rnx

diff --git a/Engine/Source/Developer/Localization/Private/PortableObjectPipeline.cpp b/Engine/Source/Developer/Localization/Private/PortableObjectPipeline.cpp
--- a/Engine/Source/Developer/Localization/Private/PortableObjectPipeline.cpp
+++ b/Engine/Source/Developer/Localization/Private/PortableObjectPipeline.cpp
@@ -310,6 +310,8 @@ namespace
 			}
 		}
 
+		UE_LOG(LogPortableObjectPipeline, Log, TEXT("Imported file %s"), *InPOFilePath);
+
 		if (bModifiedArchive)
 		{
 			// Trim any dead entries out of the archive
@@ -440,6 +442,8 @@ namespace
 			return false;
 		}
 
+		UE_LOG(LogPortableObjectPipeline, Log, TEXT("Exported file %s"), *InPOFilePath);
+
 		return true;
 	}
 }
diff --git a/Engine/Source/Developer/Localization/Private/UserGeneratedContentLocalization.cpp b/Engine/Source/Developer/Localization/Private/UserGeneratedContentLocalization.cpp
--- a/Engine/Source/Developer/Localization/Private/UserGeneratedContentLocalization.cpp
+++ b/Engine/Source/Developer/Localization/Private/UserGeneratedContentLocalization.cpp
@@ -194,19 +194,21 @@ bool ExportLocalization(TArrayView<const TSharedRef<IPlugin>> Plugins, const FEx
 					FString DestinationFilename = FilenameOrDirectory;
 					if (DestinationFilename.ReplaceInline(*PluginLocalizationTargetDirectory, *PluginLocalizationScratchDirectory) > 0)
 					{
-						IFileManager::Get().Copy(*DestinationFilename, FilenameOrDirectory);
-					}
-					else
-					{
-						bCopiedAllFiles = false;
-						return false;
+						if (IFileManager::Get().Copy(*DestinationFilename, FilenameOrDirectory) == COPY_OK)
+						{
+							UE_LOG(LogLocalization, Log, TEXT("Imported existing .po file for '%s': %s"), *Plugin->GetName(), FilenameOrDirectory);
+						}
+						else
+						{
+							bCopiedAllFiles = false;
+							UE_LOG(LogLocalization, Warning, TEXT("Failed to import existing .po file for '%s': %s"), *Plugin->GetName(), FilenameOrDirectory);
+						}
 					}
 				}
 				return true;
 			});
 			if (!bCopiedAllFiles)
 			{
-				UE_LOG(LogLocalization, Error, TEXT("Failed to import existing .po files for '%s'"), *Plugin->GetName());
 				return false;
 			}
 		}
@@ -397,6 +399,7 @@ bool ExportLocalization(TArrayView<const TSharedRef<IPlugin>> Plugins, const FEx
 			if (!ExportOptions.UGCLocDescriptor.ToJsonFile(*UGCLocFilename))
 			{
 				PostWriteFileWithSCC(UGCLocFilename);
+				UE_LOG(LogLocalization, Log, TEXT("Updated .ugcloc file for '%s': %s"), *Plugin->GetName(), *UGCLocFilename);
 			}
 			else
 			{
@@ -415,6 +418,7 @@ bool ExportLocalization(TArrayView<const TSharedRef<IPlugin>> Plugins, const FEx
 					if (IFileManager::Get().Copy(*DestinationFilename, FilenameOrDirectory) == COPY_OK)
 					{
 						PostWriteFileWithSCC(DestinationFilename);
+						UE_LOG(LogLocalization, Log, TEXT("Updated .po file for '%s': %s"), *Plugin->GetName(), *DestinationFilename);
 					}
 					else
 					{
@@ -456,6 +460,7 @@ bool ExportLocalization(TArrayView<const TSharedRef<IPlugin>> Plugins, const FEx
 				if (Plugin->UpdateDescriptor(PluginDescriptor, DescriptorUpdateFailureReason))
 				{
 					PostWriteFileWithSCC(Plugin->GetDescriptorFileName());
+					UE_LOG(LogLocalization, Log, TEXT("Updated .uplugin file for '%s'"), *Plugin->GetName());
 				}
 				else
 				{
