Added API to cleanup leftover UGC localization data

diff --git a/Engine/Source/Developer/Localization/Private/UserGeneratedContentLocalization.cpp b/Engine/Source/Developer/Localization/Private/UserGeneratedContentLocalization.cpp
--- a/Engine/Source/Developer/Localization/Private/UserGeneratedContentLocalization.cpp
+++ b/Engine/Source/Developer/Localization/Private/UserGeneratedContentLocalization.cpp
@@ -4,11 +4,13 @@
 
 #include "HAL/FileManager.h"
 #include "Misc/FileHelper.h"
+#include "Misc/MessageDialog.h"
 #include "Misc/PathViews.h"
 #include "Misc/ScopeExit.h"
 #include "JsonObjectConverter.h"
 #include "Interfaces/IPluginManager.h"
 
+#include "Internationalization/Culture.h"
 #include "Internationalization/CultureFilter.h"
 #include "Internationalization/TextLocalizationManager.h"
 #include "TextLocalizationResourceGenerator.h"
@@ -17,6 +19,8 @@
 #include "PortableObjectFormatDOM.h"
 #include "SourceControlHelpers.h"
 
+#define LOCTEXT_NAMESPACE "UserGeneratedContentLocalization"
+
 void FUserGeneratedContentLocalizationDescriptor::InitializeFromProject(const ELocalizedTextSourceCategory LocalizationCategory)
 {
 	ELocalizationLoadFlags LoadFlags = ELocalizationLoadFlags::None;
@@ -146,7 +150,12 @@ void PostWriteFileWithSCC(const FString& Filename)
 	if (USourceControlHelpers::IsEnabled())
 	{
 		// If the file didn't exist before then this will add it, otherwise it will do nothing
-		if (!USourceControlHelpers::CheckOutOrAddFile(Filename))
+		if (USourceControlHelpers::CheckOutOrAddFile(Filename))
+		{
+			// Discard the checkout if the file has no changes
+			USourceControlHelpers::RevertUnchangedFile(Filename, /*bSilent*/true);
+		}
+		else
 		{
 			UE_LOG(LogLocalization, Error, TEXT("Failed to check out file '%s'. %s"), *Filename, *USourceControlHelpers::LastErrorMsg().ToString());
 		}
@@ -396,7 +405,7 @@ bool ExportLocalization(TArrayView<const TSharedRef<IPlugin>> Plugins, const FEx
 			const FString UGCLocFilename = FPaths::Combine(PluginLocalizationTargetDirectory, FString::Printf(TEXT("%s.ugcloc"), *Plugin->GetName()));
 
 			PreWriteFileWithSCC(UGCLocFilename);
-			if (!ExportOptions.UGCLocDescriptor.ToJsonFile(*UGCLocFilename))
+			if (ExportOptions.UGCLocDescriptor.ToJsonFile(*UGCLocFilename))
 			{
 				PostWriteFileWithSCC(UGCLocFilename);
 				UE_LOG(LogLocalization, Log, TEXT("Updated .ugcloc file for '%s': %s"), *Plugin->GetName(), *UGCLocFilename);
@@ -675,4 +684,108 @@ ELoadLocalizationResult LoadLocalization(const FString& PluginName, const FStrin
 	return ELoadLocalizationResult::Success;
 }
 
+void CleanupLocalization(TArrayView<const TSharedRef<IPlugin>> Plugins, const FUserGeneratedContentLocalizationDescriptor& DefaultDescriptor, const bool bSilent)
+{
+	// Make sure we also consider localization for the native culture
+	TArray<FString> CulturesToGenerate = DefaultDescriptor.CulturesToGenerate;
+	if (!DefaultDescriptor.NativeCulture.IsEmpty())
+	{
+		CulturesToGenerate.AddUnique(DefaultDescriptor.NativeCulture);
+	}
+
+	TArray<FString> LocalizationFilesToCleanup;
+	for (const TSharedRef<IPlugin>& Plugin : Plugins)
+	{
+		const FString PluginLocalizationTargetDirectory = GetLocalizationTargetDirectory(Plugin);
+
+		// Find any leftover PO files to cleanup
+		const FString PluginPOFilename = Plugin->GetName() + TEXT(".po");
+		IFileManager::Get().IterateDirectory(*PluginLocalizationTargetDirectory, [&LocalizationFilesToCleanup, &PluginPOFilename, &CulturesToGenerate](const TCHAR* FilenameOrDirectory, bool bIsDirectory) -> bool
+		{
+			if (bIsDirectory)
+			{
+				// Note: This looks for PO files rather than the folders, as the folders may just be empty vestiges from a P4 sync without rmdir set
+				const FString PluginPOFile = FilenameOrDirectory / PluginPOFilename;
+				if (FPaths::FileExists(PluginPOFile))
+				{
+					const FString LocalizationFolder = FPaths::GetCleanFilename(FilenameOrDirectory);
+					const FString CanonicalName = FCulture::GetCanonicalName(LocalizationFolder);
+					if (!CulturesToGenerate.Contains(CanonicalName))
+					{
+						LocalizationFilesToCleanup.Add(PluginPOFile);
+					}
+				}
+			}
+			return true;
+		});
+
+		// If we aren't exporting any cultures, then also cleanup any existing descriptor file
+		if (CulturesToGenerate.Num() == 0)
+		{
+			const FString PluginUGCLocFilename = FPaths::Combine(PluginLocalizationTargetDirectory, FString::Printf(TEXT("%s.ugcloc"), *Plugin->GetName()));
+			if (FPaths::FileExists(PluginUGCLocFilename))
+			{
+				LocalizationFilesToCleanup.Add(PluginUGCLocFilename);
+			}
+		}
+	}
+
+	if (LocalizationFilesToCleanup.Num() > 0)
+	{
+		auto GetCleanupLocalizationMessage = [&LocalizationFilesToCleanup]()
+		{
+			FTextBuilder CleanupLocalizationMessageBuilder;
+			CleanupLocalizationMessageBuilder.AppendLine(LOCTEXT("CleanupLocalization.Message", "Would you like to cleanup the following localization data?"));
+			for (const FString& LeftoverPOFile : LocalizationFilesToCleanup)
+			{
+				CleanupLocalizationMessageBuilder.AppendLineFormat(LOCTEXT("CleanupLocalization.MessageLine", "    \u2022 {0}"), FText::AsCultureInvariant(LeftoverPOFile));
+			}
+			return CleanupLocalizationMessageBuilder.ToText();
+		};
+
+		if (bSilent || FMessageDialog::Open(EAppMsgType::YesNo, GetCleanupLocalizationMessage(), LOCTEXT("CleanupLocalization.Title", "Cleanup localization data?")) == EAppReturnType::Yes)
+		{
+			// Cleanup the files
+			if (USourceControlHelpers::IsEnabled())
+			{
+				USourceControlHelpers::MarkFilesForDelete(LocalizationFilesToCleanup);
+			}
+			else
+			{
+				for (const FString& LocalizationFileToCleanup : LocalizationFilesToCleanup)
+				{
+					IFileManager::Get().Delete(*LocalizationFileToCleanup);
+				}
+			}
+
+			// Cleanup the folders containing those files (will do nothing if the folder isn't actually empty)
+			for (const FString& LocalizationFileToCleanup : LocalizationFilesToCleanup)
+			{
+				const FString LocalizationPathToCleanup = FPaths::GetPath(LocalizationFileToCleanup);
+				IFileManager::Get().DeleteDirectory(*LocalizationPathToCleanup);
+			}
+
+			// If we aren't exporting any cultures, then also cleanup any plugin references to the localization data
+			if (CulturesToGenerate.Num() == 0)
+			{
+				for (const TSharedRef<IPlugin>& Plugin : Plugins)
+				{
+					FPluginDescriptor PluginDescriptor = Plugin->GetDescriptor();
+					if (PluginDescriptor.LocalizationTargets.RemoveAll([&Plugin](const FLocalizationTargetDescriptor& LocalizationTargetDescriptor) { return LocalizationTargetDescriptor.Name == Plugin->GetName(); }) > 0)
+					{
+						FText DescriptorUpdateFailureReason;
+						PreWriteFileWithSCC(Plugin->GetDescriptorFileName());
+						if (Plugin->UpdateDescriptor(PluginDescriptor, DescriptorUpdateFailureReason))
+						{
+							PostWriteFileWithSCC(Plugin->GetDescriptorFileName());
+						}
+					}
+				}
+			}
+		}
+	}
+}
+
 }
+
+#undef LOCTEXT_NAMESPACE
diff --git a/Engine/Source/Developer/Localization/Public/UserGeneratedContentLocalization.h b/Engine/Source/Developer/Localization/Public/UserGeneratedContentLocalization.h
--- a/Engine/Source/Developer/Localization/Public/UserGeneratedContentLocalization.h
+++ b/Engine/Source/Developer/Localization/Public/UserGeneratedContentLocalization.h
@@ -196,4 +196,13 @@ LOCALIZATION_API bool CompileLocalization(const FString& PluginName, const FStri
  */
 LOCALIZATION_API ELoadLocalizationResult LoadLocalization(const FString& PluginName, const FString& PluginContentDirectory, TSharedPtr<FLocTextHelper>& OutLocTextHelper, const FUserGeneratedContentLocalizationDescriptor* DefaultDescriptor = nullptr);
 
+/**
+ * Cleanup UGC localization that is no longer relevant based on the given descriptor.
+ *
+ * @param Plugins				The list of plugins to cleanup.
+ * @param DefaultDescriptor		The default UGC localization descriptor to filter existing UGC localization data against (things that don't pass the filter will be cleaned).
+ * @param bSilent				True to silently delete localization data that doesn't pass the filter, or false to confirm with the user (via a dialog).
+ */
+LOCALIZATION_API void CleanupLocalization(TArrayView<const TSharedRef<IPlugin>> Plugins, const FUserGeneratedContentLocalizationDescriptor& DefaultDescriptor, const bool bSilent = false);
+
 }
