Default to English when there is no localization data available

#rnx

diff --git a/Engine/Source/Developer/Localization/Private/UserGeneratedContentLocalization.cpp b/Engine/Source/Developer/Localization/Private/UserGeneratedContentLocalization.cpp
--- a/Engine/Source/Developer/Localization/Private/UserGeneratedContentLocalization.cpp
+++ b/Engine/Source/Developer/Localization/Private/UserGeneratedContentLocalization.cpp
@@ -44,6 +44,10 @@ void FUserGeneratedContentLocalizationDescriptor::InitializeFromProject(const EL
 	}
 
 	NativeCulture = FTextLocalizationManager::Get().GetNativeCultureName(LocalizationCategory);
+	if (NativeCulture.IsEmpty())
+	{
+		NativeCulture = TEXT("en");
+	}
 	CulturesToGenerate = FTextLocalizationManager::Get().GetLocalizedCultureNames(LoadFlags);
 
 	// Filter any cultures that are disabled in shipping or via UGC loc settings
