Gather Verse for UGC localization

#rnx

diff --git a/Engine/Source/Developer/Localization/Private/UserGeneratedContentLocalization.cpp b/Engine/Source/Developer/Localization/Private/UserGeneratedContentLocalization.cpp
--- a/Engine/Source/Developer/Localization/Private/UserGeneratedContentLocalization.cpp
+++ b/Engine/Source/Developer/Localization/Private/UserGeneratedContentLocalization.cpp
@@ -247,6 +247,8 @@ bool ExportLocalization(TArrayView<const TSharedRef<IPlugin>> Plugins, const FEx
 				ConfigSection.Add(TEXT("ArchiveName"), FString::Printf(TEXT("%s.archive"), *Plugin->GetName()));
 				ConfigSection.Add(TEXT("PortableObjectName"), FString::Printf(TEXT("%s.po"), *Plugin->GetName()));
 
+				ConfigSection.Add(TEXT("GatheredSourceBasePath"), FPaths::ConvertRelativePathToFull(Plugin->GetBaseDir()));
+
 				ConfigSection.Add(TEXT("CopyrightNotice"), ExportOptions.CopyrightNotice);
 
 				ConfigSection.Add(TEXT("NativeCulture"), ExportOptions.UGCLocDescriptor.NativeCulture);
@@ -310,6 +312,20 @@ bool ExportLocalization(TArrayView<const TSharedRef<IPlugin>> Plugins, const FEx
 				}
 			}
 
+#if WITH_VERSE
+			// Gather Verse
+			if (ExportOptions.bGatherVerse && Plugin->CanContainVerse())
+			{
+				FConfigSection& ConfigSection = GatherConfig.GatherTextStep(GatherStepIndex++);
+				ConfigSection.Add(TEXT("CommandletClass"), TEXT("GatherTextFromVerse"));
+
+				ConfigSection.Add(TEXT("IncludePathFilters"), FPaths::ConvertRelativePathToFull(FPaths::Combine(Plugin->GetBaseDir(), TEXT("*"))));
+
+				ConfigSection.Add(TEXT("ExcludePathFilters"), FPaths::ConvertRelativePathToFull(FPaths::Combine(Plugin->GetContentDir(), TEXT("Localization"), TEXT("*"))));
+				ConfigSection.Add(TEXT("ExcludePathFilters"), FPaths::ConvertRelativePathToFull(FPaths::Combine(Plugin->GetContentDir(), TEXT("L10N"), TEXT("*"))));
+			}
+#endif	// WITH_VERSE
+
 			// Generate manifest
 			{
 				FConfigSection& ConfigSection = GatherConfig.GatherTextStep(GatherStepIndex++);
@@ -383,10 +399,10 @@ bool ExportLocalization(TArrayView<const TSharedRef<IPlugin>> Plugins, const FEx
 			TArray<FString> CommandletOutputLines;
 			CommandletOutput.ParseIntoArrayLines(CommandletOutputLines);
 
-			UE_LOG(LogLocalization, Log, TEXT("Localization commandlet finished with exit code %d"), ReturnCode);
+			UE_LOG(LogLocalization, Display, TEXT("Localization commandlet finished with exit code %d"), ReturnCode);
 			for (const FString& CommandletOutputLine : CommandletOutputLines)
 			{
-				UE_LOG(LogLocalization, Log, TEXT("    %s"), *CommandletOutputLine);
+				UE_LOG(LogLocalization, Display, TEXT("    %s"), *CommandletOutputLine);
 			}
 		}
 
diff --git a/Engine/Source/Developer/Localization/Public/UserGeneratedContentLocalization.h b/Engine/Source/Developer/Localization/Public/UserGeneratedContentLocalization.h
--- a/Engine/Source/Developer/Localization/Public/UserGeneratedContentLocalization.h
+++ b/Engine/Source/Developer/Localization/Public/UserGeneratedContentLocalization.h
@@ -132,6 +132,9 @@ struct FExportLocalizationOptions
 	/** True to gather localization from assets */
 	bool bGatherAssets = true;
 
+	/** True to gather localization from Verse */
+	bool bGatherVerse = true;
+
 	/** True to update the plugin descriptors (if needed) so that they contain the exported localization target */
 	bool bUpdatePluginDescriptor = true;
 
diff --git a/Engine/Source/Editor/UnrealEd/Private/Commandlets/GatherTextFromSourceCommandlet.cpp b/Engine/Source/Editor/UnrealEd/Private/Commandlets/GatherTextFromSourceCommandlet.cpp
--- a/Engine/Source/Editor/UnrealEd/Private/Commandlets/GatherTextFromSourceCommandlet.cpp
+++ b/Engine/Source/Editor/UnrealEd/Private/Commandlets/GatherTextFromSourceCommandlet.cpp
@@ -77,6 +77,14 @@ int32 UGatherTextFromSourceCommandlet::Main( const FString& Params )
 		return -1;
 	}
 
+	// GatheredSourceBasePath
+	FString GatheredSourceBasePath;
+	GetPathFromConfig(*SectionName, TEXT("GatheredSourceBasePath"), GatheredSourceBasePath, GatherTextConfigPath);
+	if (GatheredSourceBasePath.IsEmpty())
+	{
+		GatheredSourceBasePath = UGatherTextCommandletBase::GetProjectBasePath();
+	}
+
 	// SearchDirectoryPaths
 	TArray<FString> SearchDirectoryPaths;
 	GetPathArrayFromConfig(*SectionName, TEXT("SearchDirectoryPaths"), SearchDirectoryPaths, GatherTextConfigPath);
@@ -307,11 +315,9 @@ int32 UGatherTextFromSourceCommandlet::Main( const FString& Params )
 	// Parse all source files for macros and add entries to SourceParsedEntries
 	for ( FString& SourceFile : FilesToProcess)
 	{
-		const FString& ProjectBasePath = UGatherTextCommandletBase::GetProjectBasePath();
-
 		ParseCtxt.Filename = SourceFile;
 		ParseCtxt.FileTypes = ParseCtxt.Filename.EndsWith(TEXT(".ini")) ? EGatherTextSourceFileTypes::Ini : EGatherTextSourceFileTypes::Cpp;
-		FPaths::MakePathRelativeTo(ParseCtxt.Filename, *ProjectBasePath);
+		FPaths::MakePathRelativeTo(ParseCtxt.Filename, *GatheredSourceBasePath);
 		ParseCtxt.LineNumber = 0;
 		ParseCtxt.FilePlatformName = GetSplitPlatformNameFromPath(ParseCtxt.Filename);
 		ParseCtxt.LineText.Reset();
