Localization:
Sort .po file entries by source location, namespace, key and then source string.
This allows localization vendors to group .po entries by files and have better context for the translation.

#test: Generated a new Engine.po file and verified that the .po file was indeed sorted by source location.

diff --git a/Engine/Source/Developer/Localization/Private/PortableObjectFormatDOM.cpp b/Engine/Source/Developer/Localization/Private/PortableObjectFormatDOM.cpp
--- a/Engine/Source/Developer/Localization/Private/PortableObjectFormatDOM.cpp
+++ b/Engine/Source/Developer/Localization/Private/PortableObjectFormatDOM.cpp
@@ -851,9 +851,38 @@ void FPortableObjectFormatDOM::SortEntries()
 		EntryPair.Value->ReferenceComments.Sort();
 	}
 
-	// Sort by namespace, then keys, then source text.
+	// Sort by reference (source location) then namespace, then keys, then source text.
 	auto SortingPredicate = [](const TSharedPtr<FPortableObjectEntry>& A, const TSharedPtr<FPortableObjectEntry>& B) -> bool
 	{
+		// compare reference comment
+		const int32 ReferenceCommentCount = FMath::Max(A->ReferenceComments.Num(), B->ReferenceComments.Num());
+		for (int32 index = 0; index < ReferenceCommentCount; ++index)
+		{
+			// If A has no more comments, it is before B.
+			if (!A->ReferenceComments.IsValidIndex(index) && B->ReferenceComments.IsValidIndex(index))
+			{
+				return true;
+			}
+			// If B has no more comments, it is before A.
+			if (A->ReferenceComments.IsValidIndex(index) && !B->ReferenceComments.IsValidIndex(index))
+			{
+				return false;
+			}
+
+			check(A->ReferenceComments.IsValidIndex(index) && B->ReferenceComments.IsValidIndex(index));
+
+			// If A's reference is lexicographically less, it is before B.
+			if (A->ReferenceComments[index] < B->ReferenceComments[index])
+			{
+				return true;
+			}
+			// If B's comment is lexicographically less, it is before A.
+			if (A->ReferenceComments[index] > B->ReferenceComments[index])
+			{
+				return false;
+			}
+		}
+
 		// Compare namespace
 		if (A->MsgCtxt < B->MsgCtxt)
 		{
