Ignore non-zero localization commandlet exit codes if GatherText itself returns a zero exit code

#rnx

diff --git a/Engine/Source/Developer/Localization/Private/UserGeneratedContentLocalization.cpp b/Engine/Source/Developer/Localization/Private/UserGeneratedContentLocalization.cpp
--- a/Engine/Source/Developer/Localization/Private/UserGeneratedContentLocalization.cpp
+++ b/Engine/Source/Developer/Localization/Private/UserGeneratedContentLocalization.cpp
@@ -410,7 +410,13 @@ bool ExportLocalization(TArrayView<const TSharedRef<IPlugin>> Plugins, const FEx
 		// Verify the commandlet finished cleanly
 		if (ReturnCode != 0)
 		{
-			return false;
+			// The commandlet can sometimes exit with a non-zero return code for reasons unrelated to the localization export
+			// If this happens, check to see whether the GatherText commandlet itself exited with a zero return code
+			if (!CommandletOutput.Contains(TEXT("GatherText completed with exit code 0"), ESearchCase::CaseSensitive))
+			{
+				return false;
+			}
+			UE_LOG(LogLocalization, Warning, TEXT("Localization commandlet finished with a non-zero exit code, but GatherText finished with a zero exit code. Considering the export a success, but there may be errors or omissions in the exported data."));
 		}
 	}
 
diff --git a/Engine/Source/Editor/UnrealEd/Private/Commandlets/GatherTextCommandlet.cpp b/Engine/Source/Editor/UnrealEd/Private/Commandlets/GatherTextCommandlet.cpp
--- a/Engine/Source/Editor/UnrealEd/Private/Commandlets/GatherTextCommandlet.cpp
+++ b/Engine/Source/Editor/UnrealEd/Private/Commandlets/GatherTextCommandlet.cpp
@@ -132,6 +132,8 @@ int32 UGatherTextCommandlet::Main( const FString& Params )
 	}
 	UE_LOG(LogGatherTextCommandlet, Display, TEXT("Completed all steps in %.2f seconds"), FPlatformTime::Seconds() - AllCommandletExecutionStartTime);
 
+	// Note: Other things use the below log as a tracker for GatherText completing successfully - DO NOT remove or edit this line without also updating those places
+	UE_LOG(LogGatherTextCommandlet, Display, TEXT("GatherText completed with exit code 0"));
 	return 0;
 }
 
