Moved UGC localization logging to its own category

#rnx

diff --git a/Engine/Source/Developer/Localization/Private/UserGeneratedContentLocalization.cpp b/Engine/Source/Developer/Localization/Private/UserGeneratedContentLocalization.cpp
--- a/Engine/Source/Developer/Localization/Private/UserGeneratedContentLocalization.cpp
+++ b/Engine/Source/Developer/Localization/Private/UserGeneratedContentLocalization.cpp
@@ -22,6 +22,8 @@
 
 #define LOCTEXT_NAMESPACE "UserGeneratedContentLocalization"
 
+DEFINE_LOG_CATEGORY_STATIC(LogUGCLocalization, Log, All);
+
 void FUserGeneratedContentLocalizationDescriptor::InitializeFromProject(const ELocalizedTextSourceCategory LocalizationCategory)
 {
 	ELocalizationLoadFlags LoadFlags = ELocalizationLoadFlags::None;
@@ -145,7 +147,7 @@ void PreWriteFileWithSCC(const FString& Filename)
 		// If the file already already exists, then check it out before writing to it
 		if (FPaths::FileExists(Filename) && !USourceControlHelpers::CheckOutFile(Filename))
 		{
-			UE_LOG(LogLocalization, Error, TEXT("Failed to check out file '%s'. %s"), *Filename, *USourceControlHelpers::LastErrorMsg().ToString());
+			UE_LOG(LogUGCLocalization, Error, TEXT("Failed to check out file '%s'. %s"), *Filename, *USourceControlHelpers::LastErrorMsg().ToString());
 		}
 	}
 }
@@ -162,7 +164,7 @@ void PostWriteFileWithSCC(const FString& Filename)
 		}
 		else
 		{
-			UE_LOG(LogLocalization, Error, TEXT("Failed to check out file '%s'. %s"), *Filename, *USourceControlHelpers::LastErrorMsg().ToString());
+			UE_LOG(LogUGCLocalization, Error, TEXT("Failed to check out file '%s'. %s"), *Filename, *USourceControlHelpers::LastErrorMsg().ToString());
 		}
 	}
 }
@@ -171,7 +173,7 @@ bool ExportLocalization(TArrayView<const TSharedRef<IPlugin>> Plugins, const FEx
 {
 	if (ExportOptions.UGCLocDescriptor.NativeCulture.IsEmpty())
 	{
-		UE_LOG(LogLocalization, Error, TEXT("Localization export options did not have a 'NativeCulture' set"));
+		UE_LOG(LogUGCLocalization, Error, TEXT("Localization export options did not have a 'NativeCulture' set"));
 		return false;
 	}
 
@@ -210,12 +212,12 @@ bool ExportLocalization(TArrayView<const TSharedRef<IPlugin>> Plugins, const FEx
 					{
 						if (IFileManager::Get().Copy(*DestinationFilename, FilenameOrDirectory) == COPY_OK)
 						{
-							UE_LOG(LogLocalization, Log, TEXT("Imported existing .po file for '%s': %s"), *Plugin->GetName(), FilenameOrDirectory);
+							UE_LOG(LogUGCLocalization, Log, TEXT("Imported existing .po file for '%s': %s"), *Plugin->GetName(), FilenameOrDirectory);
 						}
 						else
 						{
 							bCopiedAllFiles = false;
-							UE_LOG(LogLocalization, Warning, TEXT("Failed to import existing .po file for '%s': %s"), *Plugin->GetName(), FilenameOrDirectory);
+							UE_LOG(LogUGCLocalization, Warning, TEXT("Failed to import existing .po file for '%s': %s"), *Plugin->GetName(), FilenameOrDirectory);
 						}
 					}
 				}
@@ -394,7 +396,7 @@ bool ExportLocalization(TArrayView<const TSharedRef<IPlugin>> Plugins, const FEx
 				}
 				else
 				{
-					UE_LOG(LogLocalization, Error, TEXT("Failed to write gather config for '%s': %s"), *Plugin->GetName(), *GatherConfigFilename);
+					UE_LOG(LogUGCLocalization, Error, TEXT("Failed to write gather config for '%s': %s"), *Plugin->GetName(), *GatherConfigFilename);
 					return false;
 				}
 			}
@@ -412,10 +414,10 @@ bool ExportLocalization(TArrayView<const TSharedRef<IPlugin>> Plugins, const FEx
 			TArray<FString> CommandletOutputLines;
 			CommandletOutput.ParseIntoArrayLines(CommandletOutputLines);
 
-			UE_LOG(LogLocalization, Display, TEXT("Localization commandlet finished with exit code %d"), ReturnCode);
+			UE_LOG(LogUGCLocalization, Display, TEXT("Localization commandlet finished with exit code %d"), ReturnCode);
 			for (const FString& CommandletOutputLine : CommandletOutputLines)
 			{
-				UE_LOG(LogLocalization, Display, TEXT("    %s"), *CommandletOutputLine);
+				UE_LOG(LogUGCLocalization, Display, TEXT("    %s"), *CommandletOutputLine);
 			}
 		}
 
@@ -428,7 +430,7 @@ bool ExportLocalization(TArrayView<const TSharedRef<IPlugin>> Plugins, const FEx
 			{
 				return false;
 			}
-			UE_LOG(LogLocalization, Warning, TEXT("Localization commandlet finished with a non-zero exit code, but GatherText finished with a zero exit code. Considering the export a success, but there may be errors or omissions in the exported data."));
+			UE_LOG(LogUGCLocalization, Warning, TEXT("Localization commandlet finished with a non-zero exit code, but GatherText finished with a zero exit code. Considering the export a success, but there may be errors or omissions in the exported data."));
 		}
 	}
 
@@ -447,11 +449,11 @@ bool ExportLocalization(TArrayView<const TSharedRef<IPlugin>> Plugins, const FEx
 			if (ExportOptions.UGCLocDescriptor.ToJsonFile(*UGCLocFilename))
 			{
 				PostWriteFileWithSCC(UGCLocFilename);
-				UE_LOG(LogLocalization, Log, TEXT("Updated .ugcloc file for '%s': %s"), *Plugin->GetName(), *UGCLocFilename);
+				UE_LOG(LogUGCLocalization, Log, TEXT("Updated .ugcloc file for '%s': %s"), *Plugin->GetName(), *UGCLocFilename);
 			}
 			else
 			{
-				UE_LOG(LogLocalization, Warning, TEXT("Failed to update .ugcloc file for '%s': %s"), *Plugin->GetName(), *UGCLocFilename);
+				UE_LOG(LogUGCLocalization, Warning, TEXT("Failed to update .ugcloc file for '%s': %s"), *Plugin->GetName(), *UGCLocFilename);
 			}
 		}
 
@@ -466,11 +468,11 @@ bool ExportLocalization(TArrayView<const TSharedRef<IPlugin>> Plugins, const FEx
 					if (IFileManager::Get().Copy(*DestinationFilename, FilenameOrDirectory) == COPY_OK)
 					{
 						PostWriteFileWithSCC(DestinationFilename);
-						UE_LOG(LogLocalization, Log, TEXT("Updated .po file for '%s': %s"), *Plugin->GetName(), *DestinationFilename);
+						UE_LOG(LogUGCLocalization, Log, TEXT("Updated .po file for '%s': %s"), *Plugin->GetName(), *DestinationFilename);
 					}
 					else
 					{
-						UE_LOG(LogLocalization, Warning, TEXT("Failed to update .po file for '%s': %s"), *Plugin->GetName(), *DestinationFilename);
+						UE_LOG(LogUGCLocalization, Warning, TEXT("Failed to update .po file for '%s': %s"), *Plugin->GetName(), *DestinationFilename);
 					}
 				}
 			}
@@ -508,11 +510,11 @@ bool ExportLocalization(TArrayView<const TSharedRef<IPlugin>> Plugins, const FEx
 				if (Plugin->UpdateDescriptor(PluginDescriptor, DescriptorUpdateFailureReason))
 				{
 					PostWriteFileWithSCC(Plugin->GetDescriptorFileName());
-					UE_LOG(LogLocalization, Log, TEXT("Updated .uplugin file for '%s'"), *Plugin->GetName());
+					UE_LOG(LogUGCLocalization, Log, TEXT("Updated .uplugin file for '%s'"), *Plugin->GetName());
 				}
 				else
 				{
-					UE_LOG(LogLocalization, Warning, TEXT("Failed to update .uplugin file for '%s': %s"), *Plugin->GetName(), *DescriptorUpdateFailureReason.ToString());
+					UE_LOG(LogUGCLocalization, Warning, TEXT("Failed to update .uplugin file for '%s': %s"), *Plugin->GetName(), *DescriptorUpdateFailureReason.ToString());
 				}
 			}
 		}
@@ -535,13 +537,13 @@ bool CompileLocalizationTarget(const FString& LocalizationTargetDirectory, const
 		{
 			if (!LocMeta.SaveToFile(LocalizationTargetDirectory / LocMetaName))
 			{
-				UE_LOG(LogLocalization, Error, TEXT("Failed to save LocMeta file for '%s'"), *LocTextHelper.GetTargetName());
+				UE_LOG(LogUGCLocalization, Error, TEXT("Failed to save LocMeta file for '%s'"), *LocTextHelper.GetTargetName());
 				return false;
 			}
 		}
 		else
 		{
-			UE_LOG(LogLocalization, Error, TEXT("Failed to generate LocMeta file for '%s'"), *LocTextHelper.GetTargetName());
+			UE_LOG(LogUGCLocalization, Error, TEXT("Failed to generate LocMeta file for '%s'"), *LocTextHelper.GetTargetName());
 			return false;
 		}
 	}
@@ -557,13 +559,13 @@ bool CompileLocalizationTarget(const FString& LocalizationTargetDirectory, const
 
 			if (!LocRes.SaveToFile(LocalizationTargetDirectory / CultureToGenerate / LocResName))
 			{
-				UE_LOG(LogLocalization, Error, TEXT("Failed to save LocRes file for '%s' (culture '%s')"), *LocTextHelper.GetTargetName(), *CultureToGenerate);
+				UE_LOG(LogUGCLocalization, Error, TEXT("Failed to save LocRes file for '%s' (culture '%s')"), *LocTextHelper.GetTargetName(), *CultureToGenerate);
 				return false;
 			}
 		}
 		else
 		{
-			UE_LOG(LogLocalization, Error, TEXT("Failed to generate LocRes file for '%s' (culture '%s')"), *LocTextHelper.GetTargetName(), *CultureToGenerate);
+			UE_LOG(LogUGCLocalization, Error, TEXT("Failed to generate LocRes file for '%s' (culture '%s')"), *LocTextHelper.GetTargetName(), *CultureToGenerate);
 			return false;
 		}
 	}
@@ -688,7 +690,7 @@ ELoadLocalizationResult LoadLocalization(const FString& PluginName, const FStrin
 	FUserGeneratedContentLocalizationDescriptor UGCLocDescriptor;
 	if (!UGCLocDescriptor.FromJsonFile(*UGCLocFilename))
 	{
-		UE_LOG(LogLocalization, Error, TEXT("Failed to load localization descriptor for '%s'"), *PluginName);
+		UE_LOG(LogUGCLocalization, Error, TEXT("Failed to load localization descriptor for '%s'"), *PluginName);
 		return ELoadLocalizationResult::Failed;
 	}
 
@@ -707,7 +709,7 @@ ELoadLocalizationResult LoadLocalization(const FString& PluginName, const FStrin
 	// We always process the native culture first as it's also used to populate the manifest with the source texts
 	if (!ImportPortableObject(PluginLocalizationTargetDirectory, OutLocTextHelper->GetNativeCulture(), UGCLocDescriptor.PoFormat, *OutLocTextHelper))
 	{
-		UE_LOG(LogLocalization, Error, TEXT("Failed to load PO file for '%s' (culture '%s')"), *PluginName, *OutLocTextHelper->GetNativeCulture());
+		UE_LOG(LogUGCLocalization, Error, TEXT("Failed to load PO file for '%s' (culture '%s')"), *PluginName, *OutLocTextHelper->GetNativeCulture());
 		return ELoadLocalizationResult::Failed;
 	}
 	for (const FString& CultureToGenerate : OutLocTextHelper->GetAllCultures())
@@ -719,7 +721,7 @@ ELoadLocalizationResult LoadLocalization(const FString& PluginName, const FStrin
 
 		if (!ImportPortableObject(PluginLocalizationTargetDirectory, CultureToGenerate, UGCLocDescriptor.PoFormat, *OutLocTextHelper))
 		{
-			UE_LOG(LogLocalization, Error, TEXT("Failed to load PO file for '%s' (culture '%s')"), *PluginName, *CultureToGenerate);
+			UE_LOG(LogUGCLocalization, Error, TEXT("Failed to load PO file for '%s' (culture '%s')"), *PluginName, *CultureToGenerate);
 			return ELoadLocalizationResult::Failed;
 		}
 	}
