Reduced UGC localization log spam when the export is a success

#rnx

diff --git a/Engine/Source/Developer/Localization/Private/UserGeneratedContentLocalization.cpp b/Engine/Source/Developer/Localization/Private/UserGeneratedContentLocalization.cpp
--- a/Engine/Source/Developer/Localization/Private/UserGeneratedContentLocalization.cpp
+++ b/Engine/Source/Developer/Localization/Private/UserGeneratedContentLocalization.cpp
@@ -2,6 +2,7 @@
 
 #include "UserGeneratedContentLocalization.h"
 
+#include "HAL/IConsoleManager.h"
 #include "HAL/FileManager.h"
 #include "Misc/FileHelper.h"
 #include "Misc/MessageDialog.h"
@@ -24,6 +25,14 @@
 
 DEFINE_LOG_CATEGORY_STATIC(LogUGCLocalization, Log, All);
 
+namespace UserGeneratedContentLocalization
+{
+
+bool bAlwaysExportFullGatherLog = false;
+FAutoConsoleVariableRef CExportFullGatherLog(TEXT("Localization.UGC.AlwaysExportFullGatherLog"), bAlwaysExportFullGatherLog, TEXT("True to export the full gather log from running localization commandlet, even if there we no errors"));
+
+}
+
 void FUserGeneratedContentLocalizationDescriptor::InitializeFromProject(const ELocalizedTextSourceCategory LocalizationCategory)
 {
 	ELocalizationLoadFlags LoadFlags = ELocalizationLoadFlags::None;
@@ -409,28 +418,43 @@ bool ExportLocalization(TArrayView<const TSharedRef<IPlugin>> Plugins, const FEx
 		FString CommandletOutput;
 		const int32 ReturnCode = CommandletExecutor(FString::Join(GatherConfigFilenames, TEXT(";")), CommandletOutput);
 
-		// Log the output and result of the commandlet
+		// Verify the commandlet finished cleanly
+		bool bGatherFailed = true;
+		if (ReturnCode == 0)
 		{
-			TArray<FString> CommandletOutputLines;
-			CommandletOutput.ParseIntoArrayLines(CommandletOutputLines);
-
-			UE_LOG(LogUGCLocalization, Display, TEXT("Localization commandlet finished with exit code %d"), ReturnCode);
-			for (const FString& CommandletOutputLine : CommandletOutputLines)
+			bGatherFailed = false;
+		}
+		else
+		{
+			// The commandlet can sometimes exit with a non-zero return code for reasons unrelated to the localization export
+			// If this happens, check to see whether the GatherText commandlet itself exited with a zero return code
+			if (CommandletOutput.Contains(TEXT("GatherText completed with exit code 0"), ESearchCase::CaseSensitive))
 			{
-				UE_LOG(LogUGCLocalization, Display, TEXT("    %s"), *CommandletOutputLine);
+				bGatherFailed = false;
+				UE_LOG(LogUGCLocalization, Warning, TEXT("Localization commandlet finished with a non-zero exit code, but GatherText finished with a zero exit code. Considering the export a success, but there may be errors or omissions in the exported data."));
 			}
 		}
 
-		// Verify the commandlet finished cleanly
-		if (ReturnCode != 0)
+		// Log the output and result of the commandlet
 		{
-			// The commandlet can sometimes exit with a non-zero return code for reasons unrelated to the localization export
-			// If this happens, check to see whether the GatherText commandlet itself exited with a zero return code
-			if (!CommandletOutput.Contains(TEXT("GatherText completed with exit code 0"), ESearchCase::CaseSensitive))
+			UE_LOG(LogUGCLocalization, Display, TEXT("Localization commandlet finished with exit code %d"), ReturnCode);
+
+			if (bGatherFailed || UserGeneratedContentLocalization::bAlwaysExportFullGatherLog)
 			{
-				return false;
+				TArray<FString> CommandletOutputLines;
+				CommandletOutput.ParseIntoArrayLines(CommandletOutputLines);
+
+				for (const FString& CommandletOutputLine : CommandletOutputLines)
+				{
+					UE_LOG(LogUGCLocalization, Display, TEXT("    %s"), *CommandletOutputLine);
+				}
 			}
-			UE_LOG(LogUGCLocalization, Warning, TEXT("Localization commandlet finished with a non-zero exit code, but GatherText finished with a zero exit code. Considering the export a success, but there may be errors or omissions in the exported data."));
+		}
+
+		// If the gather failed then skip the rest of the process
+		if (bGatherFailed)
+		{
+			return false;
 		}
 	}
 
