Localization cooking: Add missing reason string when FLocalizationChunkDataGenerator::GenerateChunkDataFiles fails because an archive file is missing.
#rnx

diff --git a/Engine/Source/Developer/Localization/Private/LocTextHelper.cpp b/Engine/Source/Developer/Localization/Private/LocTextHelper.cpp
--- a/Engine/Source/Developer/Localization/Private/LocTextHelper.cpp
+++ b/Engine/Source/Developer/Localization/Private/LocTextHelper.cpp
@@ -2,6 +2,7 @@
 
 #include "LocTextHelper.h"
 #include "PlatformInfo.h"
+#include "Misc/EnumClassFlags.h"
 #include "Misc/Paths.h"
 #include "Misc/FileHelper.h"
 #include "Misc/DataDrivenPlatformInfoRegistry.h"
@@ -1601,13 +1602,13 @@ TSharedPtr<FInternationalizationManifest> FLocTextHelper::LoadManifestImpl(const
 	};
 
 	// Attempt to load an existing manifest first
-	if (!!(InLoadFlags & ELocTextHelperLoadFlags::Load))
+	if (EnumHasAnyFlags(InLoadFlags, ELocTextHelperLoadFlags::Load))
 	{
 		const bool bExists = FPaths::FileExists(InManifestFilePath);
 
-		bool bLoadedAll = bExists;
 		if (bExists)
 		{
+			bool bLoadedAll = true;
 			bLoadedAll &= LoadSingleManifest(InManifestFilePath, FName());
 			{
 				// Load all per-platform manifests too
@@ -1628,26 +1629,36 @@ TSharedPtr<FInternationalizationManifest> FLocTextHelper::LoadManifestImpl(const
 					return true;
 				});
 			}
-		}
 
-		if (bLoadedAll)
-		{
-			return LocalManifest;
-		}
-
-		if (bExists)
-		{
-			// Don't allow fallback to Create if the file exists but could not be loaded
-			return nullptr;
+			if (bLoadedAll)
+			{
+				return LocalManifest;
+			}
+			else
+			{
+				// Don't allow fallback to Create if the file exists but could not be loaded
+				return nullptr;
+			}
 		}
 	}
 
-	// If we're allowed to create a manifest than we can never fail
-	if (!!(InLoadFlags & ELocTextHelperLoadFlags::Create))
+	// If we're allowed to create a manifest than we can never fail (unless file is corrupt above)
+	if (EnumHasAnyFlags(InLoadFlags, ELocTextHelperLoadFlags::Create))
 	{
 		return LocalManifest;
 	}
 
+	if (OutError)
+	{
+		if (!EnumHasAnyFlags(InLoadFlags, ELocTextHelperLoadFlags::Load))
+		{
+			*OutError = LOCTEXT("Error_LoadManifest_CallerDidNotSpecifyLoadOrCreate", "Caller did not specify either Load or Create.");
+		}
+		else
+		{
+			*OutError = FText::Format(LOCTEXT("Error_LoadManifest_MissingManifest", "Manifest '{0}' does not exist."), FText::FromString(InManifestFilePath));
+		}
+	}
 	return nullptr;
 }
 
@@ -1783,13 +1794,13 @@ TSharedPtr<FInternationalizationArchive> FLocTextHelper::LoadArchiveImpl(const F
 	};
 
 	// Attempt to load an existing archive first
-	if (!!(InLoadFlags & ELocTextHelperLoadFlags::Load))
+	if (EnumHasAnyFlags(InLoadFlags, ELocTextHelperLoadFlags::Load))
 	{
 		const bool bExists = FPaths::FileExists(InArchiveFilePath);
 
-		bool bLoadedAll = bExists;
 		if (bExists)
 		{
+			bool bLoadedAll = true;
 			bLoadedAll &= LoadSingleArchive(InArchiveFilePath);
 			{
 				// Load all per-platform archives too
@@ -1811,26 +1822,36 @@ TSharedPtr<FInternationalizationArchive> FLocTextHelper::LoadArchiveImpl(const F
 					return true;
 				});
 			}
+			if (bLoadedAll)
+			{
+				return LocalArchive;
+			}
+			else
+			{
+				// Don't allow fallback to Create if the file exists but could not be loaded
+				return nullptr;
+			}
 		}
+	}
+
+	// If we're allowed to create an archive than we can never fail (unless file is corrupt above)
+	if (EnumHasAnyFlags(InLoadFlags, ELocTextHelperLoadFlags::Create))
+	{
+		return LocalArchive;
+	}
 
-		if (bLoadedAll)
+	if (OutError)
+	{
+		if (!EnumHasAnyFlags(InLoadFlags, ELocTextHelperLoadFlags::Load))
 		{
-			return LocalArchive;
+			*OutError = LOCTEXT("Error_LoadManifest_CallerDidNotSpecifyLoadOrCreate", "Caller did not specify either Load or Create.");
 		}
-		
-		if (bExists)
+		else
 		{
-			// Don't allow fallback to Create if the file exists but could not be loaded
-			return nullptr;
+			*OutError = FText::Format(LOCTEXT("Error_LoadArchive_MissingArchive", "Archive '{0}' does not exist."), FText::FromString(InArchiveFilePath));
 		}
 	}
 
-	// If we're allowed to create a manifest than we can never fail
-	if (!!(InLoadFlags & ELocTextHelperLoadFlags::Create))
-	{
-		return LocalArchive;
-	}
-
 	return nullptr;
 }
 
