Add overload of FPortableObjectFormatDOM::ToString that can fill an existing string

#rnx

diff --git a/Engine/Source/Developer/Localization/Private/PortableObjectFormatDOM.cpp b/Engine/Source/Developer/Localization/Private/PortableObjectFormatDOM.cpp
--- a/Engine/Source/Developer/Localization/Private/PortableObjectFormatDOM.cpp
+++ b/Engine/Source/Developer/Localization/Private/PortableObjectFormatDOM.cpp
@@ -447,19 +447,24 @@ void FPortableObjectHeader::UpdateTimeStamp()
 FString FPortableObjectFormatDOM::ToString()
 {
 	FString Result;
+	ToString(Result);
+	return Result;
+}
+
+void FPortableObjectFormatDOM::ToString(FString& OutString)
+{
+	OutString.Reset();
 
 	Header.UpdateTimeStamp();
 
-	Result += Header.ToString();
-	Result += NewLineDelimiter;
+	OutString += Header.ToString();
+	OutString += NewLineDelimiter;
 
 	for( const auto& EntryPair : Entries )
 	{
-		Result += EntryPair.Value->ToString();
-		Result += NewLineDelimiter;
+		OutString += EntryPair.Value->ToString();
+		OutString += NewLineDelimiter;
 	}
-
-	return Result;
 }
 
 bool FPortableObjectFormatDOM::FromString( const FString& InStr, FText* OutErrorMsg )
diff --git a/Engine/Source/Developer/Localization/Public/PortableObjectFormatDOM.h b/Engine/Source/Developer/Localization/Public/PortableObjectFormatDOM.h
--- a/Engine/Source/Developer/Localization/Public/PortableObjectFormatDOM.h
+++ b/Engine/Source/Developer/Localization/Public/PortableObjectFormatDOM.h
@@ -335,6 +335,7 @@ public:
 	 * @return	String representing the Portable Object that can be written directly to a file.
 	 */
 	FString ToString();
+	void ToString(FString& OutString);
 
 	/**
 	 * Parses Portable Object elements from a string.
