Rename Culture to Language to UGC localization settings

#rnx

diff --git a/Engine/Source/Developer/Localization/Public/UserGeneratedContentLocalization.h b/Engine/Source/Developer/Localization/Public/UserGeneratedContentLocalization.h
--- a/Engine/Source/Developer/Localization/Public/UserGeneratedContentLocalization.h
+++ b/Engine/Source/Developer/Localization/Public/UserGeneratedContentLocalization.h
@@ -87,24 +87,24 @@ public:
 	bool FromJsonFile(const TCHAR* InFilename);
 	
 	/**
-	 * The culture that the source text is authored in.
+	 * The language that the source text is authored in.
 	 * @note You shouldn't change this once you start to localize your text.
 	 */
-	UPROPERTY(EditAnywhere, Category=Localization)
+	UPROPERTY(EditAnywhere, Category=Localization, DisplayName="Native Language")
 	FString NativeCulture;
 
 	/**
-	 * The cultures that we should generate localization data for.
-	 * @note Will implicitly always contain the native culture during export/compile.
+	 * The languages that we should generate localization data for.
+	 * @note Will implicitly always contain the native language during export/compile.
 	 */
-	UPROPERTY(EditAnywhere, Category=Localization)
+	UPROPERTY(EditAnywhere, Category=Localization, DisplayName="Languages to Generate")
 	TArray<FString> CulturesToGenerate;
 
 	/**
 	 * What format of PO file should we use?
 	 * @note You can adjust this later and we'll attempt to preserve any existing localization data by importing with the old setting prior to export.
 	 */
-	UPROPERTY(EditAnywhere, AdvancedDisplay, Category=Localization, meta=(DisplayName="PO Format"))
+	UPROPERTY(EditAnywhere, AdvancedDisplay, Category=Localization, DisplayName="PO Format")
 	EPortableObjectFormat PoFormat = EPortableObjectFormat::Unreal;
 };
 
