Don't error if there are no PO files at all when loading a UGC loc target

This can happen if the user deletes the .po files but leaves the .ugcloc file. In this case we consider it NoData.

#rnx

diff --git a/Engine/Source/Developer/Localization/Private/UserGeneratedContentLocalization.cpp b/Engine/Source/Developer/Localization/Private/UserGeneratedContentLocalization.cpp
--- a/Engine/Source/Developer/Localization/Private/UserGeneratedContentLocalization.cpp
+++ b/Engine/Source/Developer/Localization/Private/UserGeneratedContentLocalization.cpp
@@ -729,6 +729,25 @@ ELoadLocalizationResult LoadLocalization(const FString& PluginName, const FStrin
 	OutLocTextHelper = MakeShared<FLocTextHelper>(PluginLocalizationTargetDirectory, FString::Printf(TEXT("%s.manifest"), *PluginName), FString::Printf(TEXT("%s.archive"), *PluginName), UGCLocDescriptor.NativeCulture, UGCLocDescriptor.CulturesToGenerate, nullptr);
 	OutLocTextHelper->LoadAll(ELocTextHelperLoadFlags::Create);
 
+	// Do we actually have any PO files to load?
+	// If we don't then consider this to be a NoData result (rather than error), as the user likely deleted the PO files but left the UGC localization descriptor
+	{
+		bool bHasPOFiles = false;
+		for (const FString& CultureToGenerate : OutLocTextHelper->GetAllCultures())
+		{
+			const FString POFilename = PluginLocalizationTargetDirectory / CultureToGenerate / FString::Printf(TEXT("%s.po"), *OutLocTextHelper->GetTargetName());
+			if (FPaths::FileExists(POFilename))
+			{
+				bHasPOFiles = true;
+				break;
+			}
+		}
+		if (!bHasPOFiles)
+		{
+			return ELoadLocalizationResult::NoData;
+		}
+	}
+
 	// Import each PO file data, as we'll use it to generate the LocRes (via FLocTextHelper)
 	// We always process the native culture first as it's also used to populate the manifest with the source texts
 	if (!ImportPortableObject(PluginLocalizationTargetDirectory, OutLocTextHelper->GetNativeCulture(), UGCLocDescriptor.PoFormat, *OutLocTextHelper))
