Fixed some errors when source control was enabled, but pointing at the wrong workspace

#rnx

diff --git a/Engine/Source/Developer/Localization/Private/UserGeneratedContentLocalization.cpp b/Engine/Source/Developer/Localization/Private/UserGeneratedContentLocalization.cpp
--- a/Engine/Source/Developer/Localization/Private/UserGeneratedContentLocalization.cpp
+++ b/Engine/Source/Developer/Localization/Private/UserGeneratedContentLocalization.cpp
@@ -151,7 +151,7 @@ FString GetLocalizationTargetDirectory(const TSharedRef<IPlugin>& Plugin)
 
 void PreWriteFileWithSCC(const FString& Filename)
 {
-	if (USourceControlHelpers::IsEnabled())
+	if (USourceControlHelpers::IsAvailable())
 	{
 		// If the file already already exists, then check it out before writing to it
 		if (FPaths::FileExists(Filename) && !USourceControlHelpers::CheckOutFile(Filename))
@@ -163,7 +163,7 @@ void PreWriteFileWithSCC(const FString& Filename)
 
 void PostWriteFileWithSCC(const FString& Filename)
 {
-	if (USourceControlHelpers::IsEnabled())
+	if (USourceControlHelpers::IsAvailable())
 	{
 		// If the file didn't exist before then this will add it, otherwise it will do nothing
 		if (USourceControlHelpers::CheckOutOrAddFile(Filename))
