Detect broken Rich-Text tags during localization import

diff --git a/Engine/Source/Developer/Localization/Private/LocalizationConfigurationScript.cpp b/Engine/Source/Developer/Localization/Private/LocalizationConfigurationScript.cpp
--- a/Engine/Source/Developer/Localization/Private/LocalizationConfigurationScript.cpp
+++ b/Engine/Source/Developer/Localization/Private/LocalizationConfigurationScript.cpp
@@ -1146,6 +1146,7 @@ namespace LocalizationConfigurationScript
 			ConfigSection.Add( TEXT("bSkipSourceCheck"), Target->Settings.CompileSettings.SkipSourceCheck ? TEXT("true") : TEXT("false") );
 			ConfigSection.Add( TEXT("bValidateFormatPatterns"), Target->Settings.CompileSettings.ValidateFormatPatterns ? TEXT("true") : TEXT("false") );
 			ConfigSection.Add( TEXT("bValidateSafeWhitespace"), Target->Settings.CompileSettings.ValidateSafeWhitespace ? TEXT("true") : TEXT("false") );
+			ConfigSection.Add( TEXT("bValidateRichTextTags"), Target->Settings.CompileSettings.ValidateRichTextTags ? TEXT("true") : TEXT("false"));
 
 			if (Target->Settings.SupportedCulturesStatistics.IsValidIndex(Target->Settings.NativeCultureIndex))
 			{
diff --git a/Engine/Source/Developer/Localization/Private/TextLocalizationResourceGenerator.cpp b/Engine/Source/Developer/Localization/Private/TextLocalizationResourceGenerator.cpp
--- a/Engine/Source/Developer/Localization/Private/TextLocalizationResourceGenerator.cpp
+++ b/Engine/Source/Developer/Localization/Private/TextLocalizationResourceGenerator.cpp
@@ -9,6 +9,81 @@
 
 DEFINE_LOG_CATEGORY_STATIC(LogTextLocalizationResourceGenerator, Log, All);
 
+namespace UE::TextLocalizationResourceGenerator::Private
+{
+	// This function counts the number of opening rich text tags "<TagName>" and closing rich text tags "</>" in the rich text string InText
+	void CountRichTextTags(const FStringView& InText, int& OutOpeningRichTextTagCount, int& OutClosingRighTextTagCount)
+	{
+		TCHAR PreviousChar = 0;
+		TCHAR CurrentChar = 0;
+		int TagLength = 0;
+		bool bTagOpen = false;
+		for (int CharacterIndex = 0; CharacterIndex < InText.Len(); ++CharacterIndex)
+		{
+			CurrentChar = InText[CharacterIndex];
+			if (CurrentChar == TEXT('<'))
+			{
+				// Even if a tag was already "Opened" with a '<', then we assume it was not really a tag but now it might be.
+				bTagOpen = true;
+				TagLength = 0;
+			}
+			else if (bTagOpen)
+			{
+				if (CurrentChar == TEXT('>'))
+				{
+					if (InText[CharacterIndex - 1] == TEXT('/'))
+					{
+						if (TagLength == 1)
+						{
+							OutClosingRighTextTagCount++;
+						}
+						// else it is a self-closing tag and we don't count it
+					}
+					else
+					{
+						// There is one other exception: we need to support "<br>" as a self-closing tag.
+						// It was supported in Text Render Components before multi-edit editable text and we need to support it for historic reasons.
+						bool bIsBRTag = (TagLength == 2 && InText[CharacterIndex - 2] == 'b' && InText[CharacterIndex - 1] == 'r');
+
+						if (!bIsBRTag)
+						{
+							OutOpeningRichTextTagCount++;
+						}
+					}
+
+					bTagOpen = false;
+				}
+				TagLength++;
+			}
+		}
+	}
+
+	// Check that rich text tags are complete, balanced and symmetric across the source and one of its translation. (i.e. <text color="FFFFFFFF"> TEST </>)
+	bool ValidateRichTextTags(const FStringView& InSource, const FStringView& InTranslation)
+	{
+		int TranslationRichTextOpeningTagCount = 0;
+		int TranslationRichTextClosingTagCount = 0;
+		CountRichTextTags(InTranslation, TranslationRichTextOpeningTagCount, TranslationRichTextClosingTagCount);
+		if (TranslationRichTextOpeningTagCount != TranslationRichTextClosingTagCount) // The opening and closing tags are not balanced. Could be invalid.
+		{
+			// Check if the source has the same "issue". If so, we  might tolerate it as it probably is a special edgecase that 
+			// is desirable and we can't cover all edgecases (i.e. maybe it is unbalanced for concatenation).
+			int SourceRichTextOpenTagCount = 0;
+			int SourceRichTextClosedTagCount = 0;
+			CountRichTextTags(InSource, SourceRichTextOpenTagCount, SourceRichTextClosedTagCount);
+			// We consider that it is not a translation mistake but it is voluntary if the source follow the same pattern.
+			bool SourceHasTheSameCount = SourceRichTextOpenTagCount == TranslationRichTextOpeningTagCount && SourceRichTextClosedTagCount == TranslationRichTextClosingTagCount;
+			if (!SourceHasTheSameCount)
+			{
+				// Now, we can safely assume it is really a translation error.
+				return false;
+			}
+		}
+
+		return true;
+	}
+}
+
 bool FTextLocalizationResourceGenerator::GenerateLocMeta(const FLocTextHelper& InLocTextHelper, const FString& InResourceName, FTextLocalizationMetaDataResource& OutLocMeta)
 {
 	// Populate the meta-data
@@ -102,6 +177,12 @@ bool FTextLocalizationResourceGenerator::GenerateLocRes(const FLocTextHelper& In
 						UE_LOG(LogTextLocalizationResourceGenerator, Warning, TEXT("%s"), *FLocTextHelper::SanitizeLogOutput(FString::Printf(TEXT("Translation '%s' (%s,%s) has leading or trailing whitespace for '%s'"), *TranslationText.Text, *InManifestEntry->Namespace.GetString(), *Context.Key.GetString(), *InCultureToGenerate)));
 					}
 				}
+				
+				if (!bIsNativeCulture && EnumHasAnyFlags(InGenerateFlags, EGenerateLocResFlags::ValidateRichTextTags) && 
+					!UE::TextLocalizationResourceGenerator::Private::ValidateRichTextTags(InManifestEntry->Source.Text, TranslationText.Text))
+				{
+					UE_LOG(LogTextLocalizationResourceGenerator, Warning, TEXT("%s"), *FLocTextHelper::SanitizeLogOutput(FString::Printf(TEXT("Broken Rich Text Tag detected in a translation. An unbalanced tag (a complete/incomplet opening rich text tag (i.e. <TagName>) with an incomplet/complete closing tag(</>)) was detected in the translation but not in the source text. Find the problematic tag in the translation and fix it in the PO file to remove this warning. \tTarget: '%s'\tCulture: '%s'\tSource: '%s'\tTranslation: '%s' (%s,%s)"), *InLocTextHelper.GetTargetPath(), *InCultureToGenerate, *InManifestEntry->Source.Text, *TranslationText.Text, *InManifestEntry->Namespace.GetString(), *Context.Key.GetString())));
+				}
 
 				// Find the LocRes to update
 				FTextLocalizationResource* LocResToUpdate = &OutPlatformAgnosticLocRes;
diff --git a/Engine/Source/Developer/Localization/Public/LocalizationTargetTypes.h b/Engine/Source/Developer/Localization/Public/LocalizationTargetTypes.h
--- a/Engine/Source/Developer/Localization/Public/LocalizationTargetTypes.h
+++ b/Engine/Source/Developer/Localization/Public/LocalizationTargetTypes.h
@@ -343,6 +343,7 @@ struct FLocalizationCompilationSettings
 		: SkipSourceCheck(false)
 		, ValidateFormatPatterns(true)
 		, ValidateSafeWhitespace(false)
+		, ValidateRichTextTags(false)
 	{
 	}
 
@@ -357,6 +358,10 @@ struct FLocalizationCompilationSettings
 	/* Should we validate that text doesn't contain any unsafe whitespace (leading or trailing whitespace) that could get lost during the translation process. */
 	UPROPERTY(config, EditAnywhere, Category="Validation")
 	bool ValidateSafeWhitespace;
+
+	/* Should we validate translated text containing rich text tags to see if the translated text have broken rich text tags. (i.e. "Example of <b>broken tag</") */
+	UPROPERTY(config, EditAnywhere, Category = "Validation")
+	bool ValidateRichTextTags;
 };
 
 USTRUCT()
diff --git a/Engine/Source/Developer/Localization/Public/TextLocalizationResourceGenerator.h b/Engine/Source/Developer/Localization/Public/TextLocalizationResourceGenerator.h
--- a/Engine/Source/Developer/Localization/Public/TextLocalizationResourceGenerator.h
+++ b/Engine/Source/Developer/Localization/Public/TextLocalizationResourceGenerator.h
@@ -25,6 +25,8 @@ enum class EGenerateLocResFlags : uint8
 	ValidateFormatPatterns = 1<<1,
 	/** Validate that text doesn't contain any unsafe whitespace (leading or trailing whitespace) that could get lost during the translation process */
 	ValidateSafeWhitespace = 1<<2,
+	/** Validate that rich text tags have not been broken during translation */
+	ValidateRichTextTags = 1<<3,
 };
 ENUM_CLASS_FLAGS(EGenerateLocResFlags);
 
diff --git a/Engine/Source/Editor/UnrealEd/Private/Commandlets/GenerateTextLocalizationResourceCommandlet.cpp b/Engine/Source/Editor/UnrealEd/Private/Commandlets/GenerateTextLocalizationResourceCommandlet.cpp
--- a/Engine/Source/Editor/UnrealEd/Private/Commandlets/GenerateTextLocalizationResourceCommandlet.cpp
+++ b/Engine/Source/Editor/UnrealEd/Private/Commandlets/GenerateTextLocalizationResourceCommandlet.cpp
@@ -155,6 +155,13 @@ int32 UGenerateTextLocalizationResourceCommandlet::Main(const FString& Params)
 		GenerateFlags |= (bValidateSafeWhitespace ? EGenerateLocResFlags::ValidateSafeWhitespace : EGenerateLocResFlags::None);
 	}
 
+	// Get whether to validate Rich Text tags.
+	{
+		bool bValidateRichTextTags = false;
+		GetBoolFromConfig(*SectionName, TEXT("bValidateRichTextTags"), bValidateRichTextTags, GatherTextConfigPath);
+		GenerateFlags |= (bValidateRichTextTags ? EGenerateLocResFlags::ValidateRichTextTags : EGenerateLocResFlags::None);
+	}
+
 	// Load the manifest and all archives
 	FLocTextHelper LocTextHelper(SourcePath, ManifestName, ArchiveName, NativeCultureName, CulturesToGenerate, GatherManifestHelper->GetLocFileNotifies(), GatherManifestHelper->GetPlatformSplitMode());
 	LocTextHelper.SetCopyrightNotice(GatherManifestHelper->GetCopyrightNotice());
diff --git a/Engine/Source/Programs/AutomationTool/Localization/LocalizationConfigFileBuilder.cs b/Engine/Source/Programs/AutomationTool/Localization/LocalizationConfigFileBuilder.cs
--- a/Engine/Source/Programs/AutomationTool/Localization/LocalizationConfigFileBuilder.cs
+++ b/Engine/Source/Programs/AutomationTool/Localization/LocalizationConfigFileBuilder.cs
@@ -26,6 +26,7 @@ namespace EpicGames.Localization
 			configSection.AddValue("bSkipSourceCheck", commonSettings.bSkipSourceCheck);
 			configSection.AddValue("bValidateFormatPatterns", commonSettings.bValidateFormatPatterns);
 			configSection.AddValue("bValidateSafeWhitespace", commonSettings.bValidateSafeWhitespace);
+			configSection.AddValue("bValidateRichTextTags", commonSettings.bValidateRichTextTags);
 			configSection.AddValue("NativeCulture", commonSettings.NativeCulture);
 			configSection.AddValues("CulturesToGenerate", commonSettings.CulturesToGenerate.ToArray());
 		}
diff --git a/Engine/Source/Programs/AutomationTool/Localization/LocalizationConfigFileParams.cs b/Engine/Source/Programs/AutomationTool/Localization/LocalizationConfigFileParams.cs
--- a/Engine/Source/Programs/AutomationTool/Localization/LocalizationConfigFileParams.cs
+++ b/Engine/Source/Programs/AutomationTool/Localization/LocalizationConfigFileParams.cs
@@ -17,6 +17,7 @@ namespace EpicGames.Localization
 		public bool bSkipSourceCheck { get; set; } = false;
 		public bool bValidateFormatPatterns { get; set; } = true;
 		public bool bValidateSafeWhitespace { get; set; } = false;
+		public bool bValidateRichTextTags { get; set; } = false;
 		public string NativeCulture { get; set; } = "";
 		public List<string> CulturesToGenerate { get; set; } = new List<string>();
 
@@ -39,6 +40,7 @@ namespace EpicGames.Localization
 			commonSettings.bSkipSourceCheck = false;
 			commonSettings.bValidateFormatPatterns = true;
 			commonSettings.bValidateSafeWhitespace = true;
+			commonSettings.bValidateRichTextTags = true;
 			return commonSettings;
 		}
 	}
