Extended UserGeneratedContentLocalization::CleanupLocalization to support removing all the localization data from an optional set of plugins

#rnx

diff --git a/Engine/Source/Developer/Localization/Private/UserGeneratedContentLocalization.cpp b/Engine/Source/Developer/Localization/Private/UserGeneratedContentLocalization.cpp
--- a/Engine/Source/Developer/Localization/Private/UserGeneratedContentLocalization.cpp
+++ b/Engine/Source/Developer/Localization/Private/UserGeneratedContentLocalization.cpp
@@ -772,7 +772,12 @@ ELoadLocalizationResult LoadLocalization(const FString& PluginName, const FStrin
 	return ELoadLocalizationResult::Success;
 }
 
-void CleanupLocalization(TArrayView<const TSharedRef<IPlugin>> Plugins, const FUserGeneratedContentLocalizationDescriptor& DefaultDescriptor, const bool bSilent)
+void CleanupLocalization(TArrayView<const TSharedRef<IPlugin>> PluginsToClean, const FUserGeneratedContentLocalizationDescriptor& DefaultDescriptor, const bool bSilent)
+{
+	CleanupLocalization(PluginsToClean, {}, DefaultDescriptor, bSilent);
+}
+
+void CleanupLocalization(TArrayView<const TSharedRef<IPlugin>> PluginsToClean, TArrayView<const TSharedRef<IPlugin>> PluginsToRemove, const FUserGeneratedContentLocalizationDescriptor& DefaultDescriptor, const bool bSilent)
 {
 	// Make sure we also consider localization for the native culture
 	TArray<FString> CulturesToGenerate = DefaultDescriptor.CulturesToGenerate;
@@ -781,43 +786,50 @@ void CleanupLocalization(TArrayView<const TSharedRef<IPlugin>> Plugins, const FU
 		CulturesToGenerate.AddUnique(DefaultDescriptor.NativeCulture);
 	}
 
+	// Compute the files that will be removed
 	TArray<FString> LocalizationFilesToCleanup;
-	for (const TSharedRef<IPlugin>& Plugin : Plugins)
+	auto ComputeLocalizationFilesToCleanup = [&LocalizationFilesToCleanup, &CulturesToGenerate](TArrayView<const TSharedRef<IPlugin>> Plugins, bool bRemoveAll)
 	{
-		const FString PluginLocalizationTargetDirectory = GetLocalizationTargetDirectory(Plugin);
-
-		// Find any leftover PO files to cleanup
-		const FString PluginPOFilename = Plugin->GetName() + TEXT(".po");
-		IFileManager::Get().IterateDirectory(*PluginLocalizationTargetDirectory, [&LocalizationFilesToCleanup, &PluginPOFilename, &CulturesToGenerate](const TCHAR* FilenameOrDirectory, bool bIsDirectory) -> bool
+		for (const TSharedRef<IPlugin>& Plugin : Plugins)
 		{
-			if (bIsDirectory)
+			const FString PluginLocalizationTargetDirectory = GetLocalizationTargetDirectory(Plugin);
+
+			// Find any leftover PO files to cleanup
+			const FString PluginPOFilename = Plugin->GetName() + TEXT(".po");
+			IFileManager::Get().IterateDirectory(*PluginLocalizationTargetDirectory, [bRemoveAll, &LocalizationFilesToCleanup, &PluginPOFilename, &CulturesToGenerate](const TCHAR* FilenameOrDirectory, bool bIsDirectory) -> bool
 			{
-				// Note: This looks for PO files rather than the folders, as the folders may just be empty vestiges from a P4 sync without rmdir set
-				const FString PluginPOFile = FilenameOrDirectory / PluginPOFilename;
-				if (FPaths::FileExists(PluginPOFile))
+				if (bIsDirectory)
 				{
-					const FString LocalizationFolder = FPaths::GetCleanFilename(FilenameOrDirectory);
-					const FString CanonicalName = FCulture::GetCanonicalName(LocalizationFolder);
-					if (!CulturesToGenerate.Contains(CanonicalName))
+					// Note: This looks for PO files rather than the folders, as the folders may just be empty vestiges from a P4 sync without rmdir set
+					const FString PluginPOFile = FilenameOrDirectory / PluginPOFilename;
+					if (FPaths::FileExists(PluginPOFile))
 					{
-						LocalizationFilesToCleanup.Add(PluginPOFile);
+						const FString LocalizationFolder = FPaths::GetCleanFilename(FilenameOrDirectory);
+						const FString CanonicalName = FCulture::GetCanonicalName(LocalizationFolder);
+						if (bRemoveAll || !CulturesToGenerate.Contains(CanonicalName))
+						{
+							LocalizationFilesToCleanup.Add(PluginPOFile);
+						}
 					}
 				}
-			}
-			return true;
-		});
+				return true;
+			});
 
-		// If we aren't exporting any cultures, then also cleanup any existing descriptor file
-		if (CulturesToGenerate.Num() == 0)
-		{
-			const FString PluginUGCLocFilename = FPaths::Combine(PluginLocalizationTargetDirectory, FString::Printf(TEXT("%s.ugcloc"), *Plugin->GetName()));
-			if (FPaths::FileExists(PluginUGCLocFilename))
+			// If we aren't exporting any cultures, then also cleanup any existing descriptor file
+			if (bRemoveAll || CulturesToGenerate.Num() == 0)
 			{
-				LocalizationFilesToCleanup.Add(PluginUGCLocFilename);
+				const FString PluginUGCLocFilename = FPaths::Combine(PluginLocalizationTargetDirectory, FString::Printf(TEXT("%s.ugcloc"), *Plugin->GetName()));
+				if (FPaths::FileExists(PluginUGCLocFilename))
+				{
+					LocalizationFilesToCleanup.Add(PluginUGCLocFilename);
+				}
 			}
 		}
-	}
+	};
+	ComputeLocalizationFilesToCleanup(PluginsToClean, /*bRemoveAll*/false);
+	ComputeLocalizationFilesToCleanup(PluginsToRemove, /*bRemoveAll*/true);
 
+	// Remove any files that are no longer needed, asking for confirmation when bSilent=false
 	if (LocalizationFilesToCleanup.Num() > 0)
 	{
 		auto GetCleanupLocalizationMessage = [&LocalizationFilesToCleanup]()
@@ -853,8 +865,8 @@ void CleanupLocalization(TArrayView<const TSharedRef<IPlugin>> Plugins, const FU
 				IFileManager::Get().DeleteDirectory(*LocalizationPathToCleanup);
 			}
 
-			// If we aren't exporting any cultures, then also cleanup any plugin references to the localization data
-			if (CulturesToGenerate.Num() == 0)
+			// Remove any leftover localization references from the plugins
+			auto RemovePluginLocalizationReferences = [](TArrayView<const TSharedRef<IPlugin>> Plugins)
 			{
 				for (const TSharedRef<IPlugin>& Plugin : Plugins)
 				{
@@ -869,7 +881,13 @@ void CleanupLocalization(TArrayView<const TSharedRef<IPlugin>> Plugins, const FU
 						}
 					}
 				}
+			};
+			if (CulturesToGenerate.Num() == 0)
+			{
+				// If we aren't exporting any cultures, then also cleanup any plugin references to the localization data
+				RemovePluginLocalizationReferences(PluginsToClean);
 			}
+			RemovePluginLocalizationReferences(PluginsToRemove);
 		}
 	}
 }
diff --git a/Engine/Source/Developer/Localization/Public/UserGeneratedContentLocalization.h b/Engine/Source/Developer/Localization/Public/UserGeneratedContentLocalization.h
--- a/Engine/Source/Developer/Localization/Public/UserGeneratedContentLocalization.h
+++ b/Engine/Source/Developer/Localization/Public/UserGeneratedContentLocalization.h
@@ -202,10 +202,12 @@ LOCALIZATION_API ELoadLocalizationResult LoadLocalization(const FString& PluginN
 /**
  * Cleanup UGC localization that is no longer relevant based on the given descriptor.
  *
- * @param Plugins				The list of plugins to cleanup.
+ * @param PluginsToClean		The list of plugins to cleanup.
+ * @param PluginsToRemove		The list of plugins to remove any localization data for, regardless of the UGC localization descriptor.
  * @param DefaultDescriptor		The default UGC localization descriptor to filter existing UGC localization data against (things that don't pass the filter will be cleaned).
  * @param bSilent				True to silently delete localization data that doesn't pass the filter, or false to confirm with the user (via a dialog).
  */
-LOCALIZATION_API void CleanupLocalization(TArrayView<const TSharedRef<IPlugin>> Plugins, const FUserGeneratedContentLocalizationDescriptor& DefaultDescriptor, const bool bSilent = false);
+LOCALIZATION_API void CleanupLocalization(TArrayView<const TSharedRef<IPlugin>> PluginsToClean, const FUserGeneratedContentLocalizationDescriptor& DefaultDescriptor, const bool bSilent = false);
+LOCALIZATION_API void CleanupLocalization(TArrayView<const TSharedRef<IPlugin>> PluginsToClean, TArrayView<const TSharedRef<IPlugin>> PluginsToRemove, const FUserGeneratedContentLocalizationDescriptor& DefaultDescriptor, const bool bSilent = false);
 
 }
