Improved the UGC localization API

Exposed new utility functions for querying localization target names/paths for a given plugin:
* GetLocalizationTargetName
* GetLocalizationTargetDirectory
* GetLocalizationTargetUGCLocFile
* GetLocalizationTargetPOFile
Updated existing code to use these utilities where possible, to avoid multiple places making assumptions about the naming scheme and layout of UGC localization targets.

Exposed some source control utilities for managing UGC localization targets:
* PreWriteFileWithSCC
* PostWriteFileWithSCC

Breaking: Updated some lower-level functions to take the UGC localization target path rather than the plugin content path, allowing them to also be used for UGC localization targets outside of a plugin (eg, in a scratch directory):
* CompileLocalization
* LoadLocalization
If you were previously calling these functions with the plugin content path, then you will now need to run that path through GetLocalizationTargetDirectory first.

Fixed the CompileLocalization overload taking a plugin internally passing through the wrong path to the CompileLocalization overload taking a path.

diff --git a/Engine/Source/Developer/Localization/Private/UserGeneratedContentLocalization.cpp b/Engine/Source/Developer/Localization/Private/UserGeneratedContentLocalization.cpp
--- a/Engine/Source/Developer/Localization/Private/UserGeneratedContentLocalization.cpp
+++ b/Engine/Source/Developer/Localization/Private/UserGeneratedContentLocalization.cpp
@@ -139,14 +139,32 @@ FString GetLocalizationScratchDirectory(const TSharedRef<IPlugin>& Plugin)
 	return FPaths::Combine(GetLocalizationScratchDirectory(), Plugin->GetName());
 }
 
+FString GetLocalizationTargetName(const TSharedRef<IPlugin>& Plugin)
+{
+	// Note: If you change this naming scheme you'll need to handle backwards computability with existing data,
+	//       and other code (such as the cooker) that assumes it can use the plugin name as the target name
+	//       (@see GetLocalizationTargetDirectory overload that takes the PluginName and PluginContentDirectory)
+	return Plugin->GetName();
+}
+
+FString GetLocalizationTargetDirectory(const TSharedRef<IPlugin>& Plugin)
+{
+	return FPaths::Combine(Plugin->GetContentDir(), TEXT("Localization"), Plugin->GetName());
+}
+
 FString GetLocalizationTargetDirectory(const FString& PluginName, const FString& PluginContentDirectory)
 {
 	return FPaths::Combine(PluginContentDirectory, TEXT("Localization"), PluginName);
 }
 
-FString GetLocalizationTargetDirectory(const TSharedRef<IPlugin>& Plugin)
+FString GetLocalizationTargetUGCLocFile(const FString& LocalizationTargetName, const FString& LocalizationTargetDirectory)
+{
+	return FPaths::Combine(LocalizationTargetDirectory, LocalizationTargetName + TEXT(".ugcloc"));
+}
+
+FString GetLocalizationTargetPOFile(const FString& LocalizationTargetName, const FString& LocalizationTargetDirectory, const FString& Culture)
 {
-	return GetLocalizationTargetDirectory(Plugin->GetName(), Plugin->GetContentDir());
+	return FPaths::Combine(LocalizationTargetDirectory, Culture, LocalizationTargetName + TEXT(".po"));
 }
 
 void PreWriteFileWithSCC(const FString& Filename)
@@ -206,13 +224,14 @@ bool ExportLocalization(TArrayView<const TSharedRef<IPlugin>> Plugins, const FEx
 	TArray<FString, TInlineAllocator<1>> GatherConfigFilenames;
 	for (const TSharedRef<IPlugin>& Plugin : Plugins)
 	{
+		const FString PluginLocalizationTargetName = GetLocalizationTargetName(Plugin);
 		const FString PluginLocalizationScratchDirectory = GetLocalizationScratchDirectory(Plugin);
 		const FString PluginLocalizationTargetDirectory = GetLocalizationTargetDirectory(Plugin);
 
 		// Seed the scratch directory with the current PO files for this plugin, so that the loc gather will import and preserve any existing translation data
 		{
 			bool bCopiedAllFiles = true;
-			IFileManager::Get().IterateDirectoryRecursively(*PluginLocalizationTargetDirectory, [&Plugin, &PluginLocalizationScratchDirectory, &PluginLocalizationTargetDirectory, &bCopiedAllFiles](const TCHAR* FilenameOrDirectory, bool bIsDirectory)
+			IFileManager::Get().IterateDirectoryRecursively(*PluginLocalizationTargetDirectory, [&PluginLocalizationTargetName, &PluginLocalizationScratchDirectory, &PluginLocalizationTargetDirectory, &bCopiedAllFiles](const TCHAR* FilenameOrDirectory, bool bIsDirectory)
 			{
 				if (!bIsDirectory && FPathViews::GetExtension(FilenameOrDirectory) == TEXTVIEW("po"))
 				{
@@ -221,12 +240,12 @@ bool ExportLocalization(TArrayView<const TSharedRef<IPlugin>> Plugins, const FEx
 					{
 						if (IFileManager::Get().Copy(*DestinationFilename, FilenameOrDirectory) == COPY_OK)
 						{
-							UE_LOG(LogUGCLocalization, Log, TEXT("Imported existing .po file for '%s': %s"), *Plugin->GetName(), FilenameOrDirectory);
+							UE_LOG(LogUGCLocalization, Log, TEXT("Imported existing .po file for '%s': %s"), *PluginLocalizationTargetName, FilenameOrDirectory);
 						}
 						else
 						{
 							bCopiedAllFiles = false;
-							UE_LOG(LogUGCLocalization, Warning, TEXT("Failed to import existing .po file for '%s': %s"), *Plugin->GetName(), FilenameOrDirectory);
+							UE_LOG(LogUGCLocalization, Warning, TEXT("Failed to import existing .po file for '%s': %s"), *PluginLocalizationTargetName, FilenameOrDirectory);
 						}
 					}
 				}
@@ -255,9 +274,9 @@ bool ExportLocalization(TArrayView<const TSharedRef<IPlugin>> Plugins, const FEx
 				ConfigSection.Add(TEXT("SourcePath"), FPaths::ConvertRelativePathToFull(PluginLocalizationScratchDirectory));
 				ConfigSection.Add(TEXT("DestinationPath"), FPaths::ConvertRelativePathToFull(PluginLocalizationScratchDirectory));
 
-				ConfigSection.Add(TEXT("ManifestName"), FString::Printf(TEXT("%s.manifest"), *Plugin->GetName()));
-				ConfigSection.Add(TEXT("ArchiveName"), FString::Printf(TEXT("%s.archive"), *Plugin->GetName()));
-				ConfigSection.Add(TEXT("PortableObjectName"), FString::Printf(TEXT("%s.po"), *Plugin->GetName()));
+				ConfigSection.Add(TEXT("ManifestName"), FString::Printf(TEXT("%s.manifest"), *PluginLocalizationTargetName));
+				ConfigSection.Add(TEXT("ArchiveName"), FString::Printf(TEXT("%s.archive"), *PluginLocalizationTargetName));
+				ConfigSection.Add(TEXT("PortableObjectName"), FString::Printf(TEXT("%s.po"), *PluginLocalizationTargetName));
 
 				ConfigSection.Add(TEXT("GatheredSourceBasePath"), FPaths::ConvertRelativePathToFull(Plugin->GetBaseDir()));
 
@@ -363,7 +382,7 @@ bool ExportLocalization(TArrayView<const TSharedRef<IPlugin>> Plugins, const FEx
 				// Read the UGC localization descriptor settings that were used to generate this localization data, as we should import against those
 				FUserGeneratedContentLocalizationDescriptor UGCLocDescriptorForImport;
 				{
-					const FString UGCLocFilename = FPaths::Combine(PluginLocalizationTargetDirectory, FString::Printf(TEXT("%s.ugcloc"), *Plugin->GetName()));
+					const FString UGCLocFilename = GetLocalizationTargetUGCLocFile(PluginLocalizationTargetName, PluginLocalizationTargetDirectory);
 					if (!UGCLocDescriptorForImport.FromJsonFile(*UGCLocFilename))
 					{
 						UGCLocDescriptorForImport = ExportOptions.UGCLocDescriptor;
@@ -398,14 +417,14 @@ bool ExportLocalization(TArrayView<const TSharedRef<IPlugin>> Plugins, const FEx
 			{
 				GatherConfig.Dirty = true;
 
-				FString GatherConfigFilename = FPaths::ConvertRelativePathToFull(RootLocalizationScratchDirectory / FString::Printf(TEXT("%s.ini"), *Plugin->GetName()));
+				FString GatherConfigFilename = FPaths::ConvertRelativePathToFull(RootLocalizationScratchDirectory / FString::Printf(TEXT("%s.ini"), *PluginLocalizationTargetName));
 				if (GatherConfig.Write(GatherConfigFilename))
 				{
 					GatherConfigFilenames.Add(MoveTemp(GatherConfigFilename));
 				}
 				else
 				{
-					UE_LOG(LogUGCLocalization, Error, TEXT("Failed to write gather config for '%s': %s"), *Plugin->GetName(), *GatherConfigFilename);
+					UE_LOG(LogUGCLocalization, Error, TEXT("Failed to write gather config for '%s': %s"), *PluginLocalizationTargetName, *GatherConfigFilename);
 					return false;
 				}
 			}
@@ -461,27 +480,28 @@ bool ExportLocalization(TArrayView<const TSharedRef<IPlugin>> Plugins, const FEx
 	// Copy any updated PO files back to the plugins and write out the localization settings used to generate them
 	for (const TSharedRef<IPlugin>& Plugin : Plugins)
 	{
+		const FString PluginLocalizationTargetName = GetLocalizationTargetName(Plugin);
 		const FString PluginLocalizationScratchDirectory = GetLocalizationScratchDirectory(Plugin);
 		const FString PluginLocalizationTargetDirectory = GetLocalizationTargetDirectory(Plugin);
 
 		// Write the UGC localization descriptor settings that were used to generate this localization data
 		// That will be needed to handle compilation, but also to handle import correctly if the descriptor settings change later
 		{
-			const FString UGCLocFilename = FPaths::Combine(PluginLocalizationTargetDirectory, FString::Printf(TEXT("%s.ugcloc"), *Plugin->GetName()));
+			const FString UGCLocFilename = GetLocalizationTargetUGCLocFile(PluginLocalizationTargetName, PluginLocalizationTargetDirectory);
 
 			PreWriteFileWithSCC(UGCLocFilename);
 			if (ExportOptions.UGCLocDescriptor.ToJsonFile(*UGCLocFilename))
 			{
 				PostWriteFileWithSCC(UGCLocFilename);
-				UE_LOG(LogUGCLocalization, Log, TEXT("Updated .ugcloc file for '%s': %s"), *Plugin->GetName(), *UGCLocFilename);
+				UE_LOG(LogUGCLocalization, Log, TEXT("Updated .ugcloc file for '%s': %s"), *PluginLocalizationTargetName, *UGCLocFilename);
 			}
 			else
 			{
-				UE_LOG(LogUGCLocalization, Warning, TEXT("Failed to update .ugcloc file for '%s': %s"), *Plugin->GetName(), *UGCLocFilename);
+				UE_LOG(LogUGCLocalization, Warning, TEXT("Failed to update .ugcloc file for '%s': %s"), *PluginLocalizationTargetName, *UGCLocFilename);
 			}
 		}
 
-		IFileManager::Get().IterateDirectoryRecursively(*PluginLocalizationScratchDirectory, [&Plugin, &PluginLocalizationScratchDirectory, &PluginLocalizationTargetDirectory](const TCHAR* FilenameOrDirectory, bool bIsDirectory)
+		IFileManager::Get().IterateDirectoryRecursively(*PluginLocalizationScratchDirectory, [&PluginLocalizationTargetName, &PluginLocalizationScratchDirectory, &PluginLocalizationTargetDirectory](const TCHAR* FilenameOrDirectory, bool bIsDirectory)
 		{
 			if (!bIsDirectory && FPathViews::GetExtension(FilenameOrDirectory) == TEXTVIEW("po"))
 			{
@@ -492,11 +512,11 @@ bool ExportLocalization(TArrayView<const TSharedRef<IPlugin>> Plugins, const FEx
 					if (IFileManager::Get().Copy(*DestinationFilename, FilenameOrDirectory) == COPY_OK)
 					{
 						PostWriteFileWithSCC(DestinationFilename);
-						UE_LOG(LogUGCLocalization, Log, TEXT("Updated .po file for '%s': %s"), *Plugin->GetName(), *DestinationFilename);
+						UE_LOG(LogUGCLocalization, Log, TEXT("Updated .po file for '%s': %s"), *PluginLocalizationTargetName, *DestinationFilename);
 					}
 					else
 					{
-						UE_LOG(LogUGCLocalization, Warning, TEXT("Failed to update .po file for '%s': %s"), *Plugin->GetName(), *DestinationFilename);
+						UE_LOG(LogUGCLocalization, Warning, TEXT("Failed to update .po file for '%s': %s"), *PluginLocalizationTargetName, *DestinationFilename);
 					}
 				}
 			}
@@ -506,10 +526,10 @@ bool ExportLocalization(TArrayView<const TSharedRef<IPlugin>> Plugins, const FEx
 		if (ExportOptions.bUpdatePluginDescriptor)
 		{
 			FPluginDescriptor PluginDescriptor = Plugin->GetDescriptor();
-			if (!PluginDescriptor.LocalizationTargets.ContainsByPredicate([&Plugin](const FLocalizationTargetDescriptor& LocalizationTargetDescriptor) { return LocalizationTargetDescriptor.Name == Plugin->GetName(); }))
+			if (!PluginDescriptor.LocalizationTargets.ContainsByPredicate([&PluginLocalizationTargetName](const FLocalizationTargetDescriptor& LocalizationTargetDescriptor) { return LocalizationTargetDescriptor.Name == PluginLocalizationTargetName; }))
 			{
 				FLocalizationTargetDescriptor& LocalizationTargetDescriptor = PluginDescriptor.LocalizationTargets.AddDefaulted_GetRef();
-				LocalizationTargetDescriptor.Name = Plugin->GetName();
+				LocalizationTargetDescriptor.Name = PluginLocalizationTargetName;
 				switch (ExportOptions.LocalizationCategory)
 				{
 				case ELocalizedTextSourceCategory::Game:
@@ -534,11 +554,11 @@ bool ExportLocalization(TArrayView<const TSharedRef<IPlugin>> Plugins, const FEx
 				if (Plugin->UpdateDescriptor(PluginDescriptor, DescriptorUpdateFailureReason))
 				{
 					PostWriteFileWithSCC(Plugin->GetDescriptorFileName());
-					UE_LOG(LogUGCLocalization, Log, TEXT("Updated .uplugin file for '%s'"), *Plugin->GetName());
+					UE_LOG(LogUGCLocalization, Log, TEXT("Updated .uplugin file for '%s'"), *PluginLocalizationTargetName);
 				}
 				else
 				{
-					UE_LOG(LogUGCLocalization, Warning, TEXT("Failed to update .uplugin file for '%s': %s"), *Plugin->GetName(), *DescriptorUpdateFailureReason.ToString());
+					UE_LOG(LogUGCLocalization, Warning, TEXT("Failed to update .uplugin file for '%s': %s"), *PluginLocalizationTargetName, *DescriptorUpdateFailureReason.ToString());
 				}
 			}
 		}
@@ -604,8 +624,9 @@ bool CompileLocalization(TArrayView<const TSharedRef<IPlugin>> Plugins, const FU
 	// Localization data is stored per-plugin
 	for (const TSharedRef<IPlugin>& Plugin : Plugins)
 	{
+		const FString PluginLocalizationTargetName = GetLocalizationTargetName(Plugin);
 		const FString PluginLocalizationTargetDirectory = GetLocalizationTargetDirectory(Plugin);
-		if (!CompileLocalization(Plugin->GetName(), PluginLocalizationTargetDirectory, PluginLocalizationTargetDirectory, DefaultDescriptor))
+		if (!CompileLocalization(PluginLocalizationTargetName, PluginLocalizationTargetDirectory, PluginLocalizationTargetDirectory, DefaultDescriptor))
 		{
 			return false;
 		}
@@ -614,11 +635,11 @@ bool CompileLocalization(TArrayView<const TSharedRef<IPlugin>> Plugins, const FU
 	return true;
 }
 
-bool CompileLocalization(const FString& PluginName, const FString& PluginInputContentDirectory, const FString& PluginOutputContentDirectory, const FUserGeneratedContentLocalizationDescriptor* DefaultDescriptor)
+bool CompileLocalization(const FString& LocalizationTargetName, const FString& LocalizationTargetInputDirectory, const FString& LocalizationTargetOutputDirectory, const FUserGeneratedContentLocalizationDescriptor* DefaultDescriptor)
 {
 	// Load the localization data so that we can compile it
 	TSharedPtr<FLocTextHelper> LocTextHelper;
-	const ELoadLocalizationResult LoadResult = LoadLocalization(PluginName, PluginInputContentDirectory, LocTextHelper, DefaultDescriptor);
+	const ELoadLocalizationResult LoadResult = LoadLocalization(LocalizationTargetName, LocalizationTargetInputDirectory, LocTextHelper, DefaultDescriptor);
 	if (LoadResult == ELoadLocalizationResult::NoData)
 	{
 		// Nothing to do
@@ -631,15 +652,14 @@ bool CompileLocalization(const FString& PluginName, const FString& PluginInputCo
 	}
 	check(LoadResult == ELoadLocalizationResult::Success);
 
-	const FString PluginLocalizationOutputDirectory = GetLocalizationTargetDirectory(PluginName, PluginOutputContentDirectory);
-	return CompileLocalizationTarget(PluginLocalizationOutputDirectory, *LocTextHelper);
+	return CompileLocalizationTarget(LocalizationTargetOutputDirectory, *LocTextHelper);
 }
 
 bool ImportPortableObject(const FString& LocalizationTargetDirectory, const FString& CultureToLoad, const EPortableObjectFormat PoFormat, FLocTextHelper& LocTextHelper)
 {
 	const bool bIsNativeCulture = CultureToLoad == LocTextHelper.GetNativeCulture();
 
-	const FString POFilename = LocalizationTargetDirectory / CultureToLoad / FString::Printf(TEXT("%s.po"), *LocTextHelper.GetTargetName());
+	const FString POFilename = GetLocalizationTargetPOFile(LocTextHelper.GetTargetName(), LocalizationTargetDirectory, CultureToLoad);
 
 	FString POFileData;
 	FPortableObjectFormatDOM POFile;
@@ -699,11 +719,9 @@ bool ImportPortableObject(const FString& LocalizationTargetDirectory, const FStr
 	return true;
 }
 
-ELoadLocalizationResult LoadLocalization(const FString& PluginName, const FString& PluginContentDirectory, TSharedPtr<FLocTextHelper>& OutLocTextHelper, const FUserGeneratedContentLocalizationDescriptor* DefaultDescriptor)
+ELoadLocalizationResult LoadLocalization(const FString& LocalizationTargetName, const FString& LocalizationTargetDirectory, TSharedPtr<FLocTextHelper>& OutLocTextHelper, const FUserGeneratedContentLocalizationDescriptor* DefaultDescriptor)
 {
-	const FString PluginLocalizationTargetDirectory = GetLocalizationTargetDirectory(PluginName, PluginContentDirectory);
-
-	const FString UGCLocFilename = FPaths::Combine(PluginLocalizationTargetDirectory, FString::Printf(TEXT("%s.ugcloc"), *PluginName));
+	const FString UGCLocFilename = GetLocalizationTargetUGCLocFile(LocalizationTargetName, LocalizationTargetDirectory);
 	if (!FPaths::FileExists(UGCLocFilename))
 	{
 		// Nothing to do
@@ -714,7 +732,7 @@ ELoadLocalizationResult LoadLocalization(const FString& PluginName, const FStrin
 	FUserGeneratedContentLocalizationDescriptor UGCLocDescriptor;
 	if (!UGCLocDescriptor.FromJsonFile(*UGCLocFilename))
 	{
-		UE_LOG(LogUGCLocalization, Error, TEXT("Failed to load localization descriptor for '%s'"), *PluginName);
+		UE_LOG(LogUGCLocalization, Error, TEXT("Failed to load localization descriptor for '%s'"), *LocalizationTargetName);
 		return ELoadLocalizationResult::Failed;
 	}
 
@@ -726,7 +744,7 @@ ELoadLocalizationResult LoadLocalization(const FString& PluginName, const FStrin
 	}
 
 	// Create in-memory versions of the manifest/archives that we will populate below
-	OutLocTextHelper = MakeShared<FLocTextHelper>(PluginLocalizationTargetDirectory, FString::Printf(TEXT("%s.manifest"), *PluginName), FString::Printf(TEXT("%s.archive"), *PluginName), UGCLocDescriptor.NativeCulture, UGCLocDescriptor.CulturesToGenerate, nullptr);
+	OutLocTextHelper = MakeShared<FLocTextHelper>(LocalizationTargetDirectory, FString::Printf(TEXT("%s.manifest"), *LocalizationTargetName), FString::Printf(TEXT("%s.archive"), *LocalizationTargetName), UGCLocDescriptor.NativeCulture, UGCLocDescriptor.CulturesToGenerate, nullptr);
 	OutLocTextHelper->LoadAll(ELocTextHelperLoadFlags::Create);
 
 	// Do we actually have any PO files to load?
@@ -735,7 +753,7 @@ ELoadLocalizationResult LoadLocalization(const FString& PluginName, const FStrin
 		bool bHasPOFiles = false;
 		for (const FString& CultureToGenerate : OutLocTextHelper->GetAllCultures())
 		{
-			const FString POFilename = PluginLocalizationTargetDirectory / CultureToGenerate / FString::Printf(TEXT("%s.po"), *OutLocTextHelper->GetTargetName());
+			const FString POFilename = GetLocalizationTargetPOFile(OutLocTextHelper->GetTargetName(), LocalizationTargetDirectory, CultureToGenerate);;
 			if (FPaths::FileExists(POFilename))
 			{
 				bHasPOFiles = true;
@@ -750,9 +768,9 @@ ELoadLocalizationResult LoadLocalization(const FString& PluginName, const FStrin
 
 	// Import each PO file data, as we'll use it to generate the LocRes (via FLocTextHelper)
 	// We always process the native culture first as it's also used to populate the manifest with the source texts
-	if (!ImportPortableObject(PluginLocalizationTargetDirectory, OutLocTextHelper->GetNativeCulture(), UGCLocDescriptor.PoFormat, *OutLocTextHelper))
+	if (!ImportPortableObject(LocalizationTargetDirectory, OutLocTextHelper->GetNativeCulture(), UGCLocDescriptor.PoFormat, *OutLocTextHelper))
 	{
-		UE_LOG(LogUGCLocalization, Error, TEXT("Failed to load PO file for '%s' (culture '%s')"), *PluginName, *OutLocTextHelper->GetNativeCulture());
+		UE_LOG(LogUGCLocalization, Error, TEXT("Failed to load PO file for '%s' (culture '%s')"), *LocalizationTargetName, *OutLocTextHelper->GetNativeCulture());
 		return ELoadLocalizationResult::Failed;
 	}
 	for (const FString& CultureToGenerate : OutLocTextHelper->GetAllCultures())
@@ -762,9 +780,9 @@ ELoadLocalizationResult LoadLocalization(const FString& PluginName, const FStrin
 			continue;
 		}
 
-		if (!ImportPortableObject(PluginLocalizationTargetDirectory, CultureToGenerate, UGCLocDescriptor.PoFormat, *OutLocTextHelper))
+		if (!ImportPortableObject(LocalizationTargetDirectory, CultureToGenerate, UGCLocDescriptor.PoFormat, *OutLocTextHelper))
 		{
-			UE_LOG(LogUGCLocalization, Error, TEXT("Failed to load PO file for '%s' (culture '%s')"), *PluginName, *CultureToGenerate);
+			UE_LOG(LogUGCLocalization, Error, TEXT("Failed to load PO file for '%s' (culture '%s')"), *LocalizationTargetName, *CultureToGenerate);
 			return ELoadLocalizationResult::Failed;
 		}
 	}
@@ -792,10 +810,11 @@ void CleanupLocalization(TArrayView<const TSharedRef<IPlugin>> PluginsToClean, T
 	{
 		for (const TSharedRef<IPlugin>& Plugin : Plugins)
 		{
+			const FString PluginLocalizationTargetName = GetLocalizationTargetName(Plugin);
 			const FString PluginLocalizationTargetDirectory = GetLocalizationTargetDirectory(Plugin);
 
 			// Find any leftover PO files to cleanup
-			const FString PluginPOFilename = Plugin->GetName() + TEXT(".po");
+			const FString PluginPOFilename = PluginLocalizationTargetName + TEXT(".po");
 			IFileManager::Get().IterateDirectory(*PluginLocalizationTargetDirectory, [bRemoveAll, &LocalizationFilesToCleanup, &PluginPOFilename, &CulturesToGenerate](const TCHAR* FilenameOrDirectory, bool bIsDirectory) -> bool
 			{
 				if (bIsDirectory)
@@ -818,7 +837,7 @@ void CleanupLocalization(TArrayView<const TSharedRef<IPlugin>> PluginsToClean, T
 			// If we aren't exporting any cultures, then also cleanup any existing descriptor file
 			if (bRemoveAll || CulturesToGenerate.Num() == 0)
 			{
-				const FString PluginUGCLocFilename = FPaths::Combine(PluginLocalizationTargetDirectory, FString::Printf(TEXT("%s.ugcloc"), *Plugin->GetName()));
+				const FString PluginUGCLocFilename = GetLocalizationTargetUGCLocFile(PluginLocalizationTargetName, PluginLocalizationTargetDirectory);
 				if (FPaths::FileExists(PluginUGCLocFilename))
 				{
 					LocalizationFilesToCleanup.Add(PluginUGCLocFilename);
@@ -870,8 +889,10 @@ void CleanupLocalization(TArrayView<const TSharedRef<IPlugin>> PluginsToClean, T
 			{
 				for (const TSharedRef<IPlugin>& Plugin : Plugins)
 				{
+					const FString PluginLocalizationTargetName = GetLocalizationTargetName(Plugin);
+
 					FPluginDescriptor PluginDescriptor = Plugin->GetDescriptor();
-					if (PluginDescriptor.LocalizationTargets.RemoveAll([&Plugin](const FLocalizationTargetDescriptor& LocalizationTargetDescriptor) { return LocalizationTargetDescriptor.Name == Plugin->GetName(); }) > 0)
+					if (PluginDescriptor.LocalizationTargets.RemoveAll([&PluginLocalizationTargetName](const FLocalizationTargetDescriptor& LocalizationTargetDescriptor) { return LocalizationTargetDescriptor.Name == PluginLocalizationTargetName; }) > 0)
 					{
 						FText DescriptorUpdateFailureReason;
 						PreWriteFileWithSCC(Plugin->GetDescriptorFileName());
diff --git a/Engine/Source/Developer/Localization/Public/UserGeneratedContentLocalization.h b/Engine/Source/Developer/Localization/Public/UserGeneratedContentLocalization.h
--- a/Engine/Source/Developer/Localization/Public/UserGeneratedContentLocalization.h
+++ b/Engine/Source/Developer/Localization/Public/UserGeneratedContentLocalization.h
@@ -158,6 +158,41 @@ enum class ELoadLocalizationResult : uint8
 	Success,
 };
 
+/**
+ * Utility to get the name of the UGC localization target that would be used by the given plugin.
+ */
+LOCALIZATION_API FString GetLocalizationTargetName(const TSharedRef<IPlugin>& Plugin);
+
+/**
+ * Utility to get the name of the UGC localization target directory that would be used by the given plugin.
+ */
+LOCALIZATION_API FString GetLocalizationTargetDirectory(const TSharedRef<IPlugin>& Plugin);
+
+/**
+ * Utility to get the name of the UGC localization target that would be used by the given plugin name and directory.
+ */
+LOCALIZATION_API FString GetLocalizationTargetDirectory(const FString& PluginName, const FString& PluginContentDirectory);
+
+/**
+ * Utility to get the UGCLoc file for for the given UGC localization target name and directory.
+ */
+LOCALIZATION_API FString GetLocalizationTargetUGCLocFile(const FString& LocalizationTargetName, const FString& LocalizationTargetDirectory);
+
+/**
+ * Utility to get the PO file for for the given UGC localization target name and directory, and the given culture.
+ */
+LOCALIZATION_API FString GetLocalizationTargetPOFile(const FString& LocalizationTargetName, const FString& LocalizationTargetDirectory, const FString& Culture);
+
+/**
+ * Utility to prepare a SCC managed file for writing.
+ */
+LOCALIZATION_API void PreWriteFileWithSCC(const FString& Filename);
+
+/**
+ * Utility to update a SCC managed file after writing.
+ */
+LOCALIZATION_API void PostWriteFileWithSCC(const FString& Filename);
+
 /**
  * Export UGC localization for the given plugins.
  * 
@@ -179,25 +214,25 @@ LOCALIZATION_API bool ExportLocalization(TArrayView<const TSharedRef<IPlugin>> P
 LOCALIZATION_API bool CompileLocalization(TArrayView<const TSharedRef<IPlugin>> Plugins, const FUserGeneratedContentLocalizationDescriptor* DefaultDescriptor = nullptr);
 
 /**
- * Compile UGC localization (if present) for the given plugin, producing LocMeta and LocRes files for consumption by the engine.
+ * Compile UGC localization (if present) for the given localization target, producing LocMeta and LocRes files for consumption by the engine.
  *
- * @param PluginName					The plugin name being compiled.
- * @param PluginInputContentDirectory	The content directory we'll read the source localization data from when compiling the plugin.
- * @param PluginOutputContentDirectory	The content directory we'll write the LocMeta and LocRes data to when compiling the plugin.
- * @param DefaultDescriptor				An optional default UGC localization descriptor to validate any loaded UGC localization descriptors against prior to compiling the localization data.
+ * @param LocalizationTargetName			The target name being compiled (@see GetLocalizationTargetName).
+ * @param LocalizationTargetInputDirectory	The directory we'll read the source localization data from when compiling the plugin (@see GetLocalizationTargetDirectory).
+ * @param LocalizationTargetOutputDirectory	The directory we'll write the LocMeta and LocRes data to when compiling the plugin.
+ * @param DefaultDescriptor					An optional default UGC localization descriptor to validate any loaded UGC localization descriptors against prior to compiling the localization data.
  */
-LOCALIZATION_API bool CompileLocalization(const FString& PluginName, const FString& PluginInputContentDirectory, const FString& PluginOutputContentDirectory, const FUserGeneratedContentLocalizationDescriptor* DefaultDescriptor = nullptr);
+LOCALIZATION_API bool CompileLocalization(const FString& LocalizationTargetName, const FString& LocalizationTargetInputDirectory, const FString& LocalizationTargetOutputDirectory, const FUserGeneratedContentLocalizationDescriptor* DefaultDescriptor = nullptr);
 
 /**
- * Load UGC localization source data for the given plugin.
+ * Load UGC localization source data for the given localization target.
  * @note This is typically only needed for compilation (which does it internally), but can also be useful if you have other processes that need to read the source data.
  *
- * @param PluginName					The plugin name being loaded.
- * @param PluginInputContentDirectory	The content directory we'll read the source localization data from.
+ * @param LocalizationTargetName		The target name being loaded (@see GetLocalizationTargetName).
+ * @param LocalizationTargetDirectory	The directory we'll read the source localization data from (@see GetLocalizationTargetDirectory).
  * @param OutLocTextHelper				The LocTextHelper to fill with manifest/archive data, re-generated from the source data.
  * @param DefaultDescriptor				An optional default UGC localization descriptor to validate any loaded UGC localization descriptors against prior to loading the localization data.
  */
-LOCALIZATION_API ELoadLocalizationResult LoadLocalization(const FString& PluginName, const FString& PluginContentDirectory, TSharedPtr<FLocTextHelper>& OutLocTextHelper, const FUserGeneratedContentLocalizationDescriptor* DefaultDescriptor = nullptr);
+LOCALIZATION_API ELoadLocalizationResult LoadLocalization(const FString& LocalizationTargetName, const FString& LocalizationTargetDirectory, TSharedPtr<FLocTextHelper>& OutLocTextHelper, const FUserGeneratedContentLocalizationDescriptor* DefaultDescriptor = nullptr);
 
 /**
  * Cleanup UGC localization that is no longer relevant based on the given descriptor.
diff --git a/Engine/Source/Editor/UnrealEd/Private/CookOnTheFlyServer.cpp b/Engine/Source/Editor/UnrealEd/Private/CookOnTheFlyServer.cpp
--- a/Engine/Source/Editor/UnrealEd/Private/CookOnTheFlyServer.cpp
+++ b/Engine/Source/Editor/UnrealEd/Private/CookOnTheFlyServer.cpp
@@ -12319,11 +12319,11 @@ void UCookOnTheFlyServer::CompileDLCLocalization(FBeginCookContext& BeginContext
 	});
 
 	// Compile UGC localization (if available) for this DLC plugin
-	const FString InputContentDirectory = GetContentDirectoryForDLC();
+	const FString InputLocalizationTargetDirectory = UserGeneratedContentLocalization::GetLocalizationTargetDirectory(CookByTheBookOptions->DlcName, GetContentDirectoryForDLC());
 	for (const FBeginCookContextPlatform& PlatformContext : BeginContext.PlatformContexts)
 	{
-		const FString OutputContentDirectory = ConvertToFullSandboxPath(InputContentDirectory, /*bForWrite*/true, PlatformContext.TargetPlatform->PlatformName());
-		UserGeneratedContentLocalization::CompileLocalization(CookByTheBookOptions->DlcName, InputContentDirectory, OutputContentDirectory, GetDefault<UUserGeneratedContentLocalizationSettings>()->bValidateDLCLocalizationDuringCook ? &DefaultUGCLocDescriptor : nullptr);
+		const FString OutputLocalizationTargetDirectory = ConvertToFullSandboxPath(InputLocalizationTargetDirectory, /*bForWrite*/true, PlatformContext.TargetPlatform->PlatformName());
+		UserGeneratedContentLocalization::CompileLocalization(CookByTheBookOptions->DlcName, InputLocalizationTargetDirectory, OutputLocalizationTargetDirectory, GetDefault<UUserGeneratedContentLocalizationSettings>()->bValidateDLCLocalizationDuringCook ? &DefaultUGCLocDescriptor : nullptr);
 	}
 }
 
