Support exporting UGC localization data to a scratch directory

diff --git a/Engine/Source/Developer/Localization/Private/UserGeneratedContentLocalization.cpp b/Engine/Source/Developer/Localization/Private/UserGeneratedContentLocalization.cpp
--- a/Engine/Source/Developer/Localization/Private/UserGeneratedContentLocalization.cpp
+++ b/Engine/Source/Developer/Localization/Private/UserGeneratedContentLocalization.cpp
@@ -152,6 +152,13 @@ FString GetLocalizationTargetDirectory(const TSharedRef<IPlugin>& Plugin)
 	return FPaths::Combine(Plugin->GetContentDir(), TEXT("Localization"), Plugin->GetName());
 }
 
+FString GetLocalizationTargetDirectory(const TSharedRef<IPlugin>& Plugin, const FString& LocalizationTargetRootDirectoryOverride)
+{
+	return LocalizationTargetRootDirectoryOverride.IsEmpty()
+		? GetLocalizationTargetDirectory(Plugin)
+		: FPaths::Combine(LocalizationTargetRootDirectoryOverride, Plugin->GetName());
+}
+
 FString GetLocalizationTargetDirectory(const FString& PluginName, const FString& PluginContentDirectory)
 {
 	return FPaths::Combine(PluginContentDirectory, TEXT("Localization"), PluginName);
@@ -167,9 +174,9 @@ FString GetLocalizationTargetPOFile(const FString& LocalizationTargetName, const
 	return FPaths::Combine(LocalizationTargetDirectory, Culture, LocalizationTargetName + TEXT(".po"));
 }
 
-void PreWriteFileWithSCC(const FString& Filename)
+void PreWriteFileWithSCC(const FString& Filename, const bool bUseSourceControl)
 {
-	if (USourceControlHelpers::IsAvailable())
+	if (bUseSourceControl && USourceControlHelpers::IsAvailable())
 	{
 		// If the file already already exists, then check it out before writing to it
 		if (FPaths::FileExists(Filename) && !USourceControlHelpers::CheckOutFile(Filename))
@@ -179,9 +186,9 @@ void PreWriteFileWithSCC(const FString& Filename)
 	}
 }
 
-void PostWriteFileWithSCC(const FString& Filename)
+void PostWriteFileWithSCC(const FString& Filename, const bool bUseSourceControl)
 {
-	if (USourceControlHelpers::IsAvailable())
+	if (bUseSourceControl && USourceControlHelpers::IsAvailable())
 	{
 		// If the file didn't exist before then this will add it, otherwise it will do nothing
 		if (USourceControlHelpers::CheckOutOrAddFile(Filename))
@@ -226,7 +233,7 @@ bool ExportLocalization(TArrayView<const TSharedRef<IPlugin>> Plugins, const FEx
 	{
 		const FString PluginLocalizationTargetName = GetLocalizationTargetName(Plugin);
 		const FString PluginLocalizationScratchDirectory = GetLocalizationScratchDirectory(Plugin);
-		const FString PluginLocalizationTargetDirectory = GetLocalizationTargetDirectory(Plugin);
+		const FString PluginLocalizationTargetDirectory = GetLocalizationTargetDirectory(Plugin, ExportOptions.LocalizationTargetRootDirectoryOverride);
 
 		// Seed the scratch directory with the current PO files for this plugin, so that the loc gather will import and preserve any existing translation data
 		{
@@ -482,17 +489,19 @@ bool ExportLocalization(TArrayView<const TSharedRef<IPlugin>> Plugins, const FEx
 	{
 		const FString PluginLocalizationTargetName = GetLocalizationTargetName(Plugin);
 		const FString PluginLocalizationScratchDirectory = GetLocalizationScratchDirectory(Plugin);
-		const FString PluginLocalizationTargetDirectory = GetLocalizationTargetDirectory(Plugin);
+		const FString PluginLocalizationTargetDirectory = GetLocalizationTargetDirectory(Plugin, ExportOptions.LocalizationTargetRootDirectoryOverride);
+
+		const bool bUseSourceControl = ExportOptions.LocalizationTargetRootDirectoryOverride.IsEmpty();
 
 		// Write the UGC localization descriptor settings that were used to generate this localization data
 		// That will be needed to handle compilation, but also to handle import correctly if the descriptor settings change later
 		{
 			const FString UGCLocFilename = GetLocalizationTargetUGCLocFile(PluginLocalizationTargetName, PluginLocalizationTargetDirectory);
 
-			PreWriteFileWithSCC(UGCLocFilename);
+			PreWriteFileWithSCC(UGCLocFilename, bUseSourceControl);
 			if (ExportOptions.UGCLocDescriptor.ToJsonFile(*UGCLocFilename))
 			{
-				PostWriteFileWithSCC(UGCLocFilename);
+				PostWriteFileWithSCC(UGCLocFilename, bUseSourceControl);
 				UE_LOG(LogUGCLocalization, Log, TEXT("Updated .ugcloc file for '%s': %s"), *PluginLocalizationTargetName, *UGCLocFilename);
 			}
 			else
@@ -523,7 +532,7 @@ bool ExportLocalization(TArrayView<const TSharedRef<IPlugin>> Plugins, const FEx
 			return true;
 		});
 
-		if (ExportOptions.bUpdatePluginDescriptor)
+		if (ExportOptions.bUpdatePluginDescriptor && ExportOptions.LocalizationTargetRootDirectoryOverride.IsEmpty())
 		{
 			FPluginDescriptor PluginDescriptor = Plugin->GetDescriptor();
 			if (!PluginDescriptor.LocalizationTargets.ContainsByPredicate([&PluginLocalizationTargetName](const FLocalizationTargetDescriptor& LocalizationTargetDescriptor) { return LocalizationTargetDescriptor.Name == PluginLocalizationTargetName; }))
@@ -550,10 +559,10 @@ bool ExportLocalization(TArrayView<const TSharedRef<IPlugin>> Plugins, const FEx
 				}
 
 				FText DescriptorUpdateFailureReason;
-				PreWriteFileWithSCC(Plugin->GetDescriptorFileName());
+				PreWriteFileWithSCC(Plugin->GetDescriptorFileName(), bUseSourceControl);
 				if (Plugin->UpdateDescriptor(PluginDescriptor, DescriptorUpdateFailureReason))
 				{
-					PostWriteFileWithSCC(Plugin->GetDescriptorFileName());
+					PostWriteFileWithSCC(Plugin->GetDescriptorFileName(), bUseSourceControl);
 					UE_LOG(LogUGCLocalization, Log, TEXT("Updated .uplugin file for '%s'"), *PluginLocalizationTargetName);
 				}
 				else
@@ -563,7 +572,10 @@ bool ExportLocalization(TArrayView<const TSharedRef<IPlugin>> Plugins, const FEx
 			}
 		}
 
-		LocalizationDelegates::OnLocalizationTargetDataUpdated.Broadcast(PluginLocalizationTargetDirectory);
+		if (ExportOptions.LocalizationTargetRootDirectoryOverride.IsEmpty())
+		{
+			LocalizationDelegates::OnLocalizationTargetDataUpdated.Broadcast(PluginLocalizationTargetDirectory);
+		}
 	}
 
 	return true;
diff --git a/Engine/Source/Developer/Localization/Public/UserGeneratedContentLocalization.h b/Engine/Source/Developer/Localization/Public/UserGeneratedContentLocalization.h
--- a/Engine/Source/Developer/Localization/Public/UserGeneratedContentLocalization.h
+++ b/Engine/Source/Developer/Localization/Public/UserGeneratedContentLocalization.h
@@ -146,6 +146,9 @@ struct FExportLocalizationOptions
 
 	/** An optional copyright notice to insert into the exported files */
 	FString CopyrightNotice;
+
+	/** An optional override for the root directory that we import/export the localization target data to, generating a sub-folder for each target exported (@see GetLocalizationTargetName, @see GetLocalizationTargetDirectory) */
+	FString LocalizationTargetRootDirectoryOverride;
 };
 
 enum class ELoadLocalizationResult : uint8
@@ -164,12 +167,17 @@ enum class ELoadLocalizationResult : uint8
 LOCALIZATION_API FString GetLocalizationTargetName(const TSharedRef<IPlugin>& Plugin);
 
 /**
- * Utility to get the name of the UGC localization target directory that would be used by the given plugin.
+ * Utility to get the path of the UGC localization target directory that would be used by the given plugin.
  */
 LOCALIZATION_API FString GetLocalizationTargetDirectory(const TSharedRef<IPlugin>& Plugin);
 
 /**
- * Utility to get the name of the UGC localization target that would be used by the given plugin name and directory.
+ * Utility to get the path of the UGC localization target directory that would be used by the given plugin, based on the given root directory override (@see FExportLocalizationOptions).
+ */
+LOCALIZATION_API FString GetLocalizationTargetDirectory(const TSharedRef<IPlugin>& Plugin, const FString& LocalizationTargetRootDirectoryOverride);
+
+/**
+ * Utility to get the path of the UGC localization target that would be used by the given plugin name and directory.
  */
 LOCALIZATION_API FString GetLocalizationTargetDirectory(const FString& PluginName, const FString& PluginContentDirectory);
 
@@ -186,12 +194,12 @@ LOCALIZATION_API FString GetLocalizationTargetPOFile(const FString& Localization
 /**
  * Utility to prepare a SCC managed file for writing.
  */
-LOCALIZATION_API void PreWriteFileWithSCC(const FString& Filename);
+LOCALIZATION_API void PreWriteFileWithSCC(const FString& Filename, const bool bUseSourceControl = true);
 
 /**
  * Utility to update a SCC managed file after writing.
  */
-LOCALIZATION_API void PostWriteFileWithSCC(const FString& Filename);
+LOCALIZATION_API void PostWriteFileWithSCC(const FString& Filename, const bool bUseSourceControl = true);
 
 /**
  * Export UGC localization for the given plugins.
