Only copy modified files back from the scratch directory when running a UGC loc export

diff --git a/Engine/Source/Developer/Localization/Private/UserGeneratedContentLocalization.cpp b/Engine/Source/Developer/Localization/Private/UserGeneratedContentLocalization.cpp
--- a/Engine/Source/Developer/Localization/Private/UserGeneratedContentLocalization.cpp
+++ b/Engine/Source/Developer/Localization/Private/UserGeneratedContentLocalization.cpp
@@ -2,8 +2,10 @@
 
 #include "UserGeneratedContentLocalization.h"
 
+#include "Async/ParallelFor.h"
 #include "HAL/IConsoleManager.h"
 #include "HAL/FileManager.h"
+#include "Hash/xxhash.h"
 #include "Misc/FileHelper.h"
 #include "Misc/MessageDialog.h"
 #include "Misc/PathViews.h"
@@ -203,6 +205,55 @@ void PostWriteFileWithSCC(const FString& Filename, const bool bUseSourceControl)
 	}
 }
 
+void GetLocalizationFileHashes(const FString& BaseDirectory, TMap<FString, FXxHash64>& OutFileHashes)
+{
+	// Get the list of localization files under the given directory
+	TArray<TTuple<FString, FXxHash64>> FilesToHash;
+	IFileManager::Get().IterateDirectoryRecursively(*BaseDirectory, [&FilesToHash](const TCHAR* FilenameOrDirectory, bool bIsDirectory)
+	{
+		if (!bIsDirectory)
+		{
+			if (FStringView FileExtension = FPathViews::GetExtension(FilenameOrDirectory);
+				FileExtension == TEXTVIEW("po") || FileExtension == TEXTVIEW("ugcloc"))
+			{
+				FilesToHash.Add(MakeTuple(FilenameOrDirectory, FXxHash64()));
+			}
+		}
+		return true;
+	});
+
+	// Generate their content hashes in parallel
+	ParallelFor(TEXT("UserGeneratedContentLocalization.GetLocalizationFileHashes"), FilesToHash.Num(), 1, [&FilesToHash](int32 Index)
+	{
+		TTuple<FString, FXxHash64>& FileToHashPair = FilesToHash[Index];
+
+		FXxHash64Builder HashBuilder;
+		{
+			// Don't include the PO file header in the hash as it contains transient information (like timestamps) that we don't care about
+			bool bSkippingPOFileHeader = FPathViews::GetExtension(FileToHashPair.Key) == TEXTVIEW("po");
+
+			FFileHelper::LoadFileToStringWithLineVisitor(*FileToHashPair.Key, [&bSkippingPOFileHeader, &HashBuilder](FStringView Line)
+			{
+				if (bSkippingPOFileHeader)
+				{
+					// PO file headers end on the first empty line in the file
+					bSkippingPOFileHeader = !Line.IsEmpty();
+					return;
+				}
+				HashBuilder.Update(Line.GetData(), Line.Len() * sizeof(TCHAR));
+			});
+		}
+		FileToHashPair.Value = HashBuilder.Finalize();
+	});
+
+	// Append the content hashes to the result
+	OutFileHashes.Reserve(OutFileHashes.Num() + FilesToHash.Num());
+	for (TTuple<FString, FXxHash64>& FileToHashPair : FilesToHash)
+	{
+		OutFileHashes.Add(MoveTemp(FileToHashPair.Key), MoveTemp(FileToHashPair.Value));
+	}
+}
+
 bool ExportLocalization(TArrayView<const TSharedRef<IPlugin>> Plugins, const FExportLocalizationOptions& ExportOptions, TFunctionRef<int32(const FString&, FString&)> CommandletExecutor)
 {
 	if (ExportOptions.UGCLocDescriptor.NativeCulture.IsEmpty())
@@ -229,38 +280,33 @@ bool ExportLocalization(TArrayView<const TSharedRef<IPlugin>> Plugins, const FEx
 
 	// Localization data stored per-plugin
 	TArray<FString, TInlineAllocator<1>> GatherConfigFilenames;
+	TMap<FString, TMap<FString, FXxHash64>> PerTargetLocalizationFileHashes;
 	for (const TSharedRef<IPlugin>& Plugin : Plugins)
 	{
 		const FString PluginLocalizationTargetName = GetLocalizationTargetName(Plugin);
 		const FString PluginLocalizationScratchDirectory = GetLocalizationScratchDirectory(Plugin);
 		const FString PluginLocalizationTargetDirectory = GetLocalizationTargetDirectory(Plugin, ExportOptions.LocalizationTargetRootDirectoryOverride);
 
-		// Seed the scratch directory with the current PO files for this plugin, so that the loc gather will import and preserve any existing translation data
+		// Track the source file hashes when the export started, so that we can detect post-export whether the files have actually changed
+		TMap<FString, FXxHash64>& LocalizationFileHashes = PerTargetLocalizationFileHashes.FindOrAdd(PluginLocalizationTargetName);
+		GetLocalizationFileHashes(PluginLocalizationTargetDirectory, LocalizationFileHashes);
+
+		// Seed the scratch directory with the current localization files for this plugin, so that the loc gather will import and preserve any existing translation data
+		for (const TTuple<FString, FXxHash64>& LocalizationFileHashPair : LocalizationFileHashes)
 		{
-			bool bCopiedAllFiles = true;
-			IFileManager::Get().IterateDirectoryRecursively(*PluginLocalizationTargetDirectory, [&PluginLocalizationTargetName, &PluginLocalizationScratchDirectory, &PluginLocalizationTargetDirectory, &bCopiedAllFiles](const TCHAR* FilenameOrDirectory, bool bIsDirectory)
+			const FString& SourceFilename = LocalizationFileHashPair.Key;
+			FString DestinationFilename = SourceFilename;
+			if (DestinationFilename.ReplaceInline(*PluginLocalizationTargetDirectory, *PluginLocalizationScratchDirectory) > 0)
 			{
-				if (!bIsDirectory && FPathViews::GetExtension(FilenameOrDirectory) == TEXTVIEW("po"))
+				if (IFileManager::Get().Copy(*DestinationFilename, *SourceFilename) == COPY_OK)
 				{
-					FString DestinationFilename = FilenameOrDirectory;
-					if (DestinationFilename.ReplaceInline(*PluginLocalizationTargetDirectory, *PluginLocalizationScratchDirectory) > 0)
-					{
-						if (IFileManager::Get().Copy(*DestinationFilename, FilenameOrDirectory) == COPY_OK)
-						{
-							UE_LOG(LogUGCLocalization, Log, TEXT("Imported existing .po file for '%s': %s"), *PluginLocalizationTargetName, FilenameOrDirectory);
-						}
-						else
-						{
-							bCopiedAllFiles = false;
-							UE_LOG(LogUGCLocalization, Warning, TEXT("Failed to import existing .po file for '%s': %s"), *PluginLocalizationTargetName, FilenameOrDirectory);
-						}
-					}
+					UE_LOG(LogUGCLocalization, Log, TEXT("Imported existing file for '%s': %s"), *PluginLocalizationTargetName, *SourceFilename);
+				}
+				else
+				{
+					UE_LOG(LogUGCLocalization, Warning, TEXT("Failed to import existing file for '%s': %s"), *PluginLocalizationTargetName, *SourceFilename);
+					return false;
 				}
-				return true;
-			});
-			if (!bCopiedAllFiles)
-			{
-				return false;
 			}
 		}
 
@@ -496,41 +542,47 @@ bool ExportLocalization(TArrayView<const TSharedRef<IPlugin>> Plugins, const FEx
 		// Write the UGC localization descriptor settings that were used to generate this localization data
 		// That will be needed to handle compilation, but also to handle import correctly if the descriptor settings change later
 		{
-			const FString UGCLocFilename = GetLocalizationTargetUGCLocFile(PluginLocalizationTargetName, PluginLocalizationTargetDirectory);
+			const FString UGCLocFilename = GetLocalizationTargetUGCLocFile(PluginLocalizationTargetName, PluginLocalizationScratchDirectory);
 
 			PreWriteFileWithSCC(UGCLocFilename, bUseSourceControl);
 			if (ExportOptions.UGCLocDescriptor.ToJsonFile(*UGCLocFilename))
 			{
 				PostWriteFileWithSCC(UGCLocFilename, bUseSourceControl);
-				UE_LOG(LogUGCLocalization, Log, TEXT("Updated .ugcloc file for '%s': %s"), *PluginLocalizationTargetName, *UGCLocFilename);
 			}
 			else
 			{
-				UE_LOG(LogUGCLocalization, Warning, TEXT("Failed to update .ugcloc file for '%s': %s"), *PluginLocalizationTargetName, *UGCLocFilename);
+				UE_LOG(LogUGCLocalization, Warning, TEXT("Failed to write updated .ugcloc file for '%s': %s"), *PluginLocalizationTargetName, *UGCLocFilename);
 			}
 		}
 
-		IFileManager::Get().IterateDirectoryRecursively(*PluginLocalizationScratchDirectory, [&PluginLocalizationTargetName, &PluginLocalizationScratchDirectory, &PluginLocalizationTargetDirectory](const TCHAR* FilenameOrDirectory, bool bIsDirectory)
+		// Track the scratch file hashes now that the export has finished, so that we can detect which files have actually changed from their source
+		TMap<FString, FXxHash64> ScratchLocalizationFileHashes;
+		GetLocalizationFileHashes(PluginLocalizationScratchDirectory, ScratchLocalizationFileHashes);
+
+		const TMap<FString, FXxHash64>& SourceLocalizationFileHashes = PerTargetLocalizationFileHashes.FindChecked(PluginLocalizationTargetName);
+		for (const TTuple<FString, FXxHash64>& LocalizationFileHashPair : ScratchLocalizationFileHashes)
 		{
-			if (!bIsDirectory && FPathViews::GetExtension(FilenameOrDirectory) == TEXTVIEW("po"))
+			const FString& SourceFilename = LocalizationFileHashPair.Key;
+			FString DestinationFilename = SourceFilename;
+			if (DestinationFilename.ReplaceInline(*PluginLocalizationScratchDirectory, *PluginLocalizationTargetDirectory) > 0)
 			{
-				FString DestinationFilename = FilenameOrDirectory;
-				if (DestinationFilename.ReplaceInline(*PluginLocalizationScratchDirectory, *PluginLocalizationTargetDirectory) > 0)
+				// Only copy files with modified hashes to avoid source control churn
+				if (const FXxHash64* SourceFileHash = SourceLocalizationFileHashes.Find(DestinationFilename);
+					!SourceFileHash || *SourceFileHash != LocalizationFileHashPair.Value)
 				{
 					PreWriteFileWithSCC(DestinationFilename);
-					if (IFileManager::Get().Copy(*DestinationFilename, FilenameOrDirectory) == COPY_OK)
+					if (IFileManager::Get().Copy(*DestinationFilename, *SourceFilename) == COPY_OK)
 					{
 						PostWriteFileWithSCC(DestinationFilename);
-						UE_LOG(LogUGCLocalization, Log, TEXT("Updated .po file for '%s': %s"), *PluginLocalizationTargetName, *DestinationFilename);
+						UE_LOG(LogUGCLocalization, Log, TEXT("Updated file for '%s': %s"), *PluginLocalizationTargetName, *DestinationFilename);
 					}
 					else
 					{
-						UE_LOG(LogUGCLocalization, Warning, TEXT("Failed to update .po file for '%s': %s"), *PluginLocalizationTargetName, *DestinationFilename);
+						UE_LOG(LogUGCLocalization, Warning, TEXT("Failed to update file for '%s': %s"), *PluginLocalizationTargetName, *DestinationFilename);
 					}
 				}
 			}
-			return true;
-		});
+		}
 
 		if (ExportOptions.bUpdatePluginDescriptor && ExportOptions.LocalizationTargetRootDirectoryOverride.IsEmpty())
 		{
