Re-added missing bUseSourceControl argument to Pre/PostWriteFileWithSCC

#rnx

diff --git a/Engine/Source/Developer/Localization/Private/UserGeneratedContentLocalization.cpp b/Engine/Source/Developer/Localization/Private/UserGeneratedContentLocalization.cpp
--- a/Engine/Source/Developer/Localization/Private/UserGeneratedContentLocalization.cpp
+++ b/Engine/Source/Developer/Localization/Private/UserGeneratedContentLocalization.cpp
@@ -541,18 +541,10 @@ bool ExportLocalization(TArrayView<const TSharedRef<IPlugin>> Plugins, const FEx
 
 		// Write the UGC localization descriptor settings that were used to generate this localization data
 		// That will be needed to handle compilation, but also to handle import correctly if the descriptor settings change later
+		if (const FString UGCLocFilename = GetLocalizationTargetUGCLocFile(PluginLocalizationTargetName, PluginLocalizationScratchDirectory);
+			!ExportOptions.UGCLocDescriptor.ToJsonFile(*UGCLocFilename))
 		{
-			const FString UGCLocFilename = GetLocalizationTargetUGCLocFile(PluginLocalizationTargetName, PluginLocalizationScratchDirectory);
-
-			PreWriteFileWithSCC(UGCLocFilename, bUseSourceControl);
-			if (ExportOptions.UGCLocDescriptor.ToJsonFile(*UGCLocFilename))
-			{
-				PostWriteFileWithSCC(UGCLocFilename, bUseSourceControl);
-			}
-			else
-			{
-				UE_LOG(LogUGCLocalization, Warning, TEXT("Failed to write updated .ugcloc file for '%s': %s"), *PluginLocalizationTargetName, *UGCLocFilename);
-			}
+			UE_LOG(LogUGCLocalization, Warning, TEXT("Failed to write updated .ugcloc file for '%s': %s"), *PluginLocalizationTargetName, *UGCLocFilename);
 		}
 
 		// Track the scratch file hashes now that the export has finished, so that we can detect which files have actually changed from their source
@@ -570,10 +562,10 @@ bool ExportLocalization(TArrayView<const TSharedRef<IPlugin>> Plugins, const FEx
 				if (const FXxHash64* SourceFileHash = SourceLocalizationFileHashes.Find(DestinationFilename);
 					!SourceFileHash || *SourceFileHash != LocalizationFileHashPair.Value)
 				{
-					PreWriteFileWithSCC(DestinationFilename);
+					PreWriteFileWithSCC(DestinationFilename, bUseSourceControl);
 					if (IFileManager::Get().Copy(*DestinationFilename, *SourceFilename) == COPY_OK)
 					{
-						PostWriteFileWithSCC(DestinationFilename);
+						PostWriteFileWithSCC(DestinationFilename, bUseSourceControl);
 						UE_LOG(LogUGCLocalization, Log, TEXT("Updated file for '%s': %s"), *PluginLocalizationTargetName, *DestinationFilename);
 					}
 					else
