Cleaned up some API organization in PortableObjectFormatDOM

#rnx

diff --git a/Engine/Source/Developer/Localization/Private/PortableObjectFormatDOM.cpp b/Engine/Source/Developer/Localization/Private/PortableObjectFormatDOM.cpp
--- a/Engine/Source/Developer/Localization/Private/PortableObjectFormatDOM.cpp
+++ b/Engine/Source/Developer/Localization/Private/PortableObjectFormatDOM.cpp
@@ -944,49 +944,49 @@ void FPortableObjectFormatDOM::SortEntries()
 
 void FPortableObjectEntry::AddExtractedComment( const FString& InComment )
 {
-	if(!InComment.IsEmpty())
+	if (!InComment.IsEmpty())
 	{
 		ExtractedComments.AddUnique(InComment);
 	}
 }
 
-void FPortableObjectEntry::AddReference( const FString& InReference )
+void FPortableObjectEntry::AddExtractedComments( const TArray<FString>& InComments )
 {
-	if(!InReference.IsEmpty())
+	for( const FString& ExtractedComment : InComments )
 	{
-		ReferenceComments.AddUnique(InReference);
+		AddExtractedComment( ExtractedComment );
 	}
 }
 
-void FPortableObjectEntry::AddExtractedComments( const TArray<FString>& InComments )
+void FPortableObjectEntry::AddTranslatorComment(const FString& InComment)
 {
-	for( const FString& ExtractedComment : InComments )
+	if (!InComment.IsEmpty())
 	{
-		AddExtractedComment( ExtractedComment );
+		TranslatorComments.AddUnique(InComment);
 	}
 }
 
-void FPortableObjectEntry::AddReferences( const TArray<FString>& InReferences )
+void FPortableObjectEntry::AddTranslatorComments(const TArray<FString>& InComments)
 {
-	for( const FString& Reference : InReferences )
+	for (const FString& TranslatorComment : InComments)
 	{
-		AddReference( Reference );
+		AddTranslatorComment(TranslatorComment);
 	}
 }
 
-void FPortableObjectEntry::AddTranslatorComment(const FString& InComment)
+void FPortableObjectEntry::AddReference( const FString& InReference )
 {
-	if (!InComment.IsEmpty())
+	if (!InReference.IsEmpty())
 	{
-		TranslatorComments.AddUnique(InComment);
+		ReferenceComments.AddUnique(InReference);
 	}
 }
 
-void FPortableObjectEntry::AddTranslatorComments(const TArray<FString>& InComments)
+void FPortableObjectEntry::AddReferences( const TArray<FString>& InReferences )
 {
-	for (const FString& TranslatorComment : InComments)
+	for( const FString& Reference : InReferences )
 	{
-		AddTranslatorComment(TranslatorComment);
+		AddReference( Reference );
 	}
 }
 
diff --git a/Engine/Source/Developer/Localization/Public/PortableObjectFormatDOM.h b/Engine/Source/Developer/Localization/Public/PortableObjectFormatDOM.h
--- a/Engine/Source/Developer/Localization/Public/PortableObjectFormatDOM.h
+++ b/Engine/Source/Developer/Localization/Public/PortableObjectFormatDOM.h
@@ -65,11 +65,6 @@ public:
 	 */
 	FString Variant() const;
 
-	/**
-	 * Retrieves the culture equivalent that matches the set language code.
-	 */
-	FCulturePtr GetCulture() const;
-
 	/**
 	 * Retrieves the display name for the language.
 	 */
@@ -151,17 +146,10 @@ public:
 	/**
 	 * Helper function that adds to the extracted comments.
 	 *
-	 * @param InComment	String representing an extracted comment entry.
+	 * @param InComment		String representing an extracted comment entry.
 	 */
 	void AddExtractedComment( const FString& InComment );
 
-	/**
-	 * Helper function that adds to the reference comments.
-	 *
-	 * @param InReference	String representing reference comment entries of the following form: "File/path/file.cpp:20 File/path/file.cpp:21 File/path/file2.cpp:5".
-	 */
-	void AddReference( const FString& InReference );
-
 	/**
 	 * Helper function that adds to the extracted comments.
 	 *
@@ -172,7 +160,7 @@ public:
 	/**
 	 * Helper function that adds to the translator comments.
 	 *
-	 * @param InComment	String representing a translator comment entry.
+	 * @param InComment		String representing a translator comment entry.
 	 */
 	void AddTranslatorComment(const FString& InComment);
 
@@ -183,6 +171,13 @@ public:
 	 */
 	void AddTranslatorComments(const TArray<FString>& InComments);
 
+	/**
+	 * Helper function that adds to the reference comments.
+	 *
+	 * @param InReference	String representing reference comment entries of the following form: "File/path/file.cpp:20 File/path/file.cpp:21 File/path/file2.cpp:5".
+	 */
+	void AddReference( const FString& InReference );
+
 	/**
 	 * Helper function that adds to the reference comments.
 	 *
@@ -230,8 +225,7 @@ public:
 
 	/* Stores flags.  Lines starting with #,  
 	    
-		#: /Engine/Source/Runtime/Engine/Private/Actor.cpp:2306
-		#: /Engine/Source/Runtime/Engine/Private/Actor.cpp:2307 /Engine/Source/Runtime/Engine/Private/Actor.cpp:2308
+		#, fuzzy
 		msgid "The Quick Brown Fox Jumps Over The Lazy Dog"
 		msgstr ""
 	*/
