Implemented support for merging two UGC localization targets together

This can also be invoked when running an export to an override directory, to allow any existing project data to be merged into the override directory during the export

diff --git a/Engine/Source/Developer/Localization/Private/UserGeneratedContentLocalization.cpp b/Engine/Source/Developer/Localization/Private/UserGeneratedContentLocalization.cpp
--- a/Engine/Source/Developer/Localization/Private/UserGeneratedContentLocalization.cpp
+++ b/Engine/Source/Developer/Localization/Private/UserGeneratedContentLocalization.cpp
@@ -304,11 +304,19 @@ bool ExportLocalization(TArrayView<const TSharedRef<IPlugin>> Plugins, const FEx
 				}
 				else
 				{
-					UE_LOG(LogUGCLocalization, Warning, TEXT("Failed to import existing file for '%s': %s"), *PluginLocalizationTargetName, *SourceFilename);
+					UE_LOG(LogUGCLocalization, Error, TEXT("Failed to import existing file for '%s': %s"), *PluginLocalizationTargetName, *SourceFilename);
 					return false;
 				}
 			}
 		}
+		if (!ExportOptions.LocalizationTargetRootDirectoryOverride.IsEmpty() && ExportOptions.MergeProjectDataWithRootDirectoryOverrideData.IsSet())
+		{
+			const FString PluginLocalizationSourceDirectory = GetLocalizationTargetDirectory(Plugin);
+			if (!MergeLocalization(PluginLocalizationTargetName, PluginLocalizationSourceDirectory, PluginLocalizationTargetName, PluginLocalizationScratchDirectory, ExportOptions.MergeProjectDataWithRootDirectoryOverrideData.GetValue(), /*bUseSourceControl*/false))
+			{
+				return false;
+			}
+		}
 
 		// Build the gather config
 		{
@@ -846,6 +854,182 @@ ELoadLocalizationResult LoadLocalization(const FString& LocalizationTargetName,
 	return ELoadLocalizationResult::Success;
 }
 
+bool MergeLocalization(const FString& SourceLocalizationTargetName, const FString& SourceLocalizationTargetDirectory, const FString& DestLocalizationTargetName, const FString& DestLocalizationTargetDirectory, const EMergeLocalizationMode MergeMode, const bool bUseSourceControl)
+{
+	TMap<FString, FXxHash64> SourceLocalizationFileHashes;
+	GetLocalizationFileHashes(SourceLocalizationTargetDirectory, SourceLocalizationFileHashes);
+
+	if (SourceLocalizationFileHashes.IsEmpty())
+	{
+		// Nothing to do
+		return true;
+	}
+
+	TMap<FString, FXxHash64> DestLocalizationFileHashes;
+	GetLocalizationFileHashes(DestLocalizationTargetDirectory, DestLocalizationFileHashes);
+
+	auto SourceFilenameToDestFilename = [&SourceLocalizationTargetDirectory, &DestLocalizationTargetDirectory, SourcePOFilename = SourceLocalizationTargetName + TEXT(".po"), DestPOFilename = DestLocalizationTargetName + TEXT(".po")](const FString& SourceFilename)
+	{
+		FString DestFilename = SourceFilename;
+		DestFilename.ReplaceInline(*SourceLocalizationTargetDirectory, *DestLocalizationTargetDirectory);
+		DestFilename.ReplaceInline(*SourcePOFilename, *DestPOFilename);
+		return DestFilename;
+	};
+
+	auto LoadPOFile = [](const FString& POFilename, FPortableObjectFormatDOM& OutPOFile)
+	{
+		FString POFileData;
+		if (!FFileHelper::LoadFileToString(POFileData, *POFilename) || !OutPOFile.FromString(POFileData))
+		{
+			UE_LOG(LogUGCLocalization, Error, TEXT("Merge failed to load '%s'"), *POFilename);
+			return false;
+		}
+		return true;
+	};
+
+	auto SavePOFile = [bUseSourceControl](const FString& POFilename, FPortableObjectFormatDOM& POFile)
+	{
+		FString POFileData;
+		POFile.ToString(POFileData);
+
+		PreWriteFileWithSCC(POFilename, bUseSourceControl);
+		if (FFileHelper::SaveStringToFile(POFileData, *POFilename, FFileHelper::EEncodingOptions::ForceUTF8))
+		{
+			PostWriteFileWithSCC(POFilename, bUseSourceControl);
+			return true;
+		}
+		else
+		{
+			UE_LOG(LogUGCLocalization, Error, TEXT("Merge failed to save '%s'"), *POFilename);
+			return false;
+		}
+	};
+
+	struct FPOFileToMerge
+	{
+		FString SourceFilename;
+		FString DestFilename;
+		TUniquePtr<FPortableObjectFormatDOM> ModifiedPOFile;
+		bool bMergeFailed = false;
+	};
+
+	TArray<FPOFileToMerge> POFilesToMerge;
+	for (const TTuple<FString, FXxHash64>& SourceLocalizationFileHashPair : SourceLocalizationFileHashes)
+	{
+		const FString& SourceFilename = SourceLocalizationFileHashPair.Key;
+		const FString DestFilename = SourceFilenameToDestFilename(SourceFilename);
+
+		if (const FXxHash64* DestLocalizationFileHash = DestLocalizationFileHashes.Find(DestFilename))
+		{
+			if (SourceLocalizationFileHashPair.Value == *DestLocalizationFileHash)
+			{
+				// File hash is identical; nothing to merge
+				UE_LOG(LogUGCLocalization, Log, TEXT("Merge skipped '%s' as it is identical to '%s'"), *SourceFilename, *DestFilename);
+			}
+			else if (FPathViews::GetExtension(SourceFilename) == TEXTVIEW("po"))
+			{
+				// Merge source into dest
+				POFilesToMerge.Add(FPOFileToMerge{ SourceFilename, DestFilename });
+			}
+		}
+		else
+		{
+			// File doesn't exist at dest; just copy directly from source
+			PreWriteFileWithSCC(DestFilename);
+			if (IFileManager::Get().Copy(*DestFilename, *SourceFilename) == COPY_OK)
+			{
+				PostWriteFileWithSCC(DestFilename);
+				UE_LOG(LogUGCLocalization, Log, TEXT("Merge copied '%s' to '%s'"), *SourceFilename, *DestFilename);
+			}
+			else
+			{
+				UE_LOG(LogUGCLocalization, Error, TEXT("Merge failed to copy '%s' to '%s'"), *SourceFilename, *DestFilename);
+				return false;
+			}
+		}
+	}
+
+	ParallelFor(TEXT("UserGeneratedContentLocalization.MergeLocalization"), POFilesToMerge.Num(), 1, [&POFilesToMerge, &LoadPOFile, &SavePOFile, MergeMode, bUseSourceControl](int32 Index)
+	{
+		FPOFileToMerge& POFileToMerge = POFilesToMerge[Index];
+
+		FPortableObjectFormatDOM SourcePOFile;
+		TUniquePtr<FPortableObjectFormatDOM> DestPOFile = MakeUnique<FPortableObjectFormatDOM>();
+		if (!LoadPOFile(POFileToMerge.SourceFilename, SourcePOFile) || !LoadPOFile(POFileToMerge.DestFilename, *DestPOFile))
+		{
+			POFileToMerge.bMergeFailed = true;
+			return;
+		}
+
+		bool bModifiedDestPOFile = false;
+		for (auto SourceEntryPairIter = SourcePOFile.GetEntriesIterator(); SourceEntryPairIter; ++SourceEntryPairIter)
+		{
+			const TSharedPtr<FPortableObjectEntry>& SourcePOEntry = SourceEntryPairIter->Value;
+			if (SourcePOEntry->MsgId.IsEmpty() || SourcePOEntry->MsgStr.Num() == 0 || SourcePOEntry->MsgStr[0].IsEmpty())
+			{
+				// We ignore the header entry or entries with no translation
+				continue;
+			}
+
+			if (TSharedPtr<FPortableObjectEntry> DestPOEntry = DestPOFile->FindEntry(SourcePOEntry.ToSharedRef()))
+			{
+				// Replace the entry in dest?
+				if (MergeMode == EMergeLocalizationMode::All || DestPOEntry->MsgStr.Num() == 0 || DestPOEntry->MsgStr[0].IsEmpty())
+				{
+					DestPOEntry->MsgStr = SourcePOEntry->MsgStr;
+					bModifiedDestPOFile = true;
+				}
+			}
+			else
+			{
+				// Add a new entry to dest
+				DestPOFile->AddEntry(MakeShared<FPortableObjectEntry>(*SourcePOEntry));
+				bModifiedDestPOFile = true;
+			}
+		}
+
+		if (bModifiedDestPOFile)
+		{
+			if (bUseSourceControl)
+			{
+				// When using source control, we need to defer the save until after the ParallelFor
+				POFileToMerge.ModifiedPOFile = MoveTemp(DestPOFile);
+			}
+			else
+			{
+				// When not using source control, we can handle the save within the ParallelFor
+				if (!SavePOFile(POFileToMerge.DestFilename, *DestPOFile))
+				{
+					POFileToMerge.bMergeFailed = true;
+					return;
+				}
+			}
+		}
+	});
+
+	for (FPOFileToMerge& POFileToMerge : POFilesToMerge)
+	{
+		// When using source control, we need to defer the save until after the ParallelFor
+		if (POFileToMerge.ModifiedPOFile)
+		{
+			check(bUseSourceControl);
+			if (!SavePOFile(POFileToMerge.DestFilename, *POFileToMerge.ModifiedPOFile))
+			{
+				POFileToMerge.bMergeFailed = true;
+			}
+		}
+
+		if (POFileToMerge.bMergeFailed)
+		{
+			return false;
+		}
+
+		UE_LOG(LogUGCLocalization, Log, TEXT("Merge applied '%s' to '%s'"), *POFileToMerge.SourceFilename, *POFileToMerge.DestFilename);
+	}
+
+	return true;
+}
+
 void CleanupLocalization(TArrayView<const TSharedRef<IPlugin>> PluginsToClean, const FUserGeneratedContentLocalizationDescriptor& DefaultDescriptor, const bool bSilent)
 {
 	CleanupLocalization(PluginsToClean, {}, DefaultDescriptor, bSilent);
diff --git a/Engine/Source/Developer/Localization/Public/UserGeneratedContentLocalization.h b/Engine/Source/Developer/Localization/Public/UserGeneratedContentLocalization.h
--- a/Engine/Source/Developer/Localization/Public/UserGeneratedContentLocalization.h
+++ b/Engine/Source/Developer/Localization/Public/UserGeneratedContentLocalization.h
@@ -6,6 +6,7 @@
 #include "Containers/ArrayView.h"
 #include "Containers/Map.h"
 #include "Containers/UnrealString.h"
+#include "Misc/Optional.h"
 #include "Templates/Function.h"
 #include "Templates/SharedPointer.h"
 
@@ -118,6 +119,8 @@ public:
 namespace UserGeneratedContentLocalization
 {
 
+enum class EMergeLocalizationMode : uint8;
+
 struct FExportLocalizationOptions
 {
 	/** Common export options for all plugins */
@@ -149,6 +152,9 @@ struct FExportLocalizationOptions
 
 	/** An optional override for the root directory that we import/export the localization target data to, generating a sub-folder for each target exported (@see GetLocalizationTargetName, @see GetLocalizationTargetDirectory) */
 	FString LocalizationTargetRootDirectoryOverride;
+
+	/** If set, merge any existing plugin localization data with the data copied from LocalizationTargetRootDirectoryOverride (in the export scratch directory) prior to running the export */
+	TOptional<EMergeLocalizationMode> MergeProjectDataWithRootDirectoryOverrideData;
 };
 
 enum class ELoadLocalizationResult : uint8
@@ -161,6 +167,20 @@ enum class ELoadLocalizationResult : uint8
 	Success,
 };
 
+enum class EMergeLocalizationMode : uint8
+{
+	/**
+	 * Only merge untranslated strings
+	 */
+	Untranslated,
+
+	/**
+	 * Merge all strings, even if they have an existing translation
+	 * @note This will overwrite any existing translation data
+	 */
+	All,
+};
+
 /**
  * Utility to get the name of the UGC localization target that would be used by the given plugin.
  */
@@ -242,6 +262,19 @@ LOCALIZATION_API bool CompileLocalization(const FString& LocalizationTargetName,
  */
 LOCALIZATION_API ELoadLocalizationResult LoadLocalization(const FString& LocalizationTargetName, const FString& LocalizationTargetDirectory, TSharedPtr<FLocTextHelper>& OutLocTextHelper, const FUserGeneratedContentLocalizationDescriptor* DefaultDescriptor = nullptr);
 
+/**
+ * Merge the PO file data of two UGC localization targets together, so that any PO entries with translations from source are present in dest. Whether the translation from source is actually applied to dest depends on the MergeMode.
+ * @note Dest does not have to currently exist prior to the merge.
+ * 
+ * @param SourceLocalizationTargetName		The name of the source localization target (@see GetLocalizationTargetName).
+ * @param SourceLocalizationTargetDirectory	The directory of the source localization target (@see GetLocalizationTargetDirectory).
+ * @param DestLocalizationTargetName		The name of the dest localization target (@see GetLocalizationTargetName).
+ * @param DestLocalizationTargetDirectory	The directory of the dest localization target (@see GetLocalizationTargetDirectory).
+ * @param MergeMode							When should an existing translation from source be applied to dest?
+ * @param bUseSourceControl					Pass false as an optimization if you knew for certain that DestLocalizationTargetDirectory is outside of source control, otherwise pass true.
+ */
+LOCALIZATION_API bool MergeLocalization(const FString& SourceLocalizationTargetName, const FString& SourceLocalizationTargetDirectory, const FString& DestLocalizationTargetName, const FString& DestLocalizationTargetDirectory, const EMergeLocalizationMode MergeMode, const bool bUseSourceControl = true);
+
 /**
  * Cleanup UGC localization that is no longer relevant based on the given descriptor.
  *
