Fixed some edge cases around PO file handling in UserGeneratedContentLocalization

* LoadLocalization will now warn (rather than error) if a PO file is missing on disk.
* PreWriteFileWithSCC will now handle the file already existing on disk, but not yet being source controlled.

diff --git a/Engine/Source/Developer/Localization/Private/UserGeneratedContentLocalization.cpp b/Engine/Source/Developer/Localization/Private/UserGeneratedContentLocalization.cpp
--- a/Engine/Source/Developer/Localization/Private/UserGeneratedContentLocalization.cpp
+++ b/Engine/Source/Developer/Localization/Private/UserGeneratedContentLocalization.cpp
@@ -5,6 +5,7 @@
 #include "Async/ParallelFor.h"
 #include "HAL/IConsoleManager.h"
 #include "HAL/FileManager.h"
+#include "HAL/PlatformFileManager.h"
 #include "Hash/xxhash.h"
 #include "Misc/FileHelper.h"
 #include "Misc/MessageDialog.h"
@@ -181,9 +182,25 @@ void PreWriteFileWithSCC(const FString& Filename, const bool bUseSourceControl)
 	if (bUseSourceControl && USourceControlHelpers::IsAvailable())
 	{
 		// If the file already already exists, then check it out before writing to it
-		if (FPaths::FileExists(Filename) && !USourceControlHelpers::CheckOutFile(Filename))
+		// We also consider an add here, as the file may have been added on disk prior to running an export (eg, when running automation that downloads the files from elsewhere)
+		if (FPaths::FileExists(Filename))
 		{
-			UE_LOG(LogUGCLocalization, Error, TEXT("Failed to check out file '%s'. %s"), *Filename, *USourceControlHelpers::LastErrorMsg().ToString());
+			if (USourceControlHelpers::CheckOutOrAddFile(Filename, /*bSilent*/true))
+			{
+				// Make sure the file is actually writable, as adding a read-only file to source control may leave it read-only
+				if (IPlatformFile& PlatformFile = FPlatformFileManager::Get().GetPlatformFile();
+					PlatformFile.IsReadOnly(*Filename))
+				{
+					if (!PlatformFile.SetReadOnly(*Filename, false))
+					{
+						UE_LOG(LogUGCLocalization, Error, TEXT("Failed to make file '%s' writable"), *Filename);
+					}
+				}
+			}
+			else
+			{
+				UE_LOG(LogUGCLocalization, Error, TEXT("Failed to check out or add file '%s'. %s"), *Filename, *USourceControlHelpers::LastErrorMsg().ToString());
+			}
 		}
 	}
 }
@@ -193,14 +210,14 @@ void PostWriteFileWithSCC(const FString& Filename, const bool bUseSourceControl)
 	if (bUseSourceControl && USourceControlHelpers::IsAvailable())
 	{
 		// If the file didn't exist before then this will add it, otherwise it will do nothing
-		if (USourceControlHelpers::CheckOutOrAddFile(Filename))
+		if (USourceControlHelpers::CheckOutOrAddFile(Filename, /*bSilent*/true))
 		{
 			// Discard the checkout if the file has no changes
 			USourceControlHelpers::RevertUnchangedFile(Filename, /*bSilent*/true);
 		}
 		else
 		{
-			UE_LOG(LogUGCLocalization, Error, TEXT("Failed to check out file '%s'. %s"), *Filename, *USourceControlHelpers::LastErrorMsg().ToString());
+			UE_LOG(LogUGCLocalization, Error, TEXT("Failed to check out or add file '%s'. %s"), *Filename, *USourceControlHelpers::LastErrorMsg().ToString());
 		}
 	}
 }
@@ -819,8 +836,8 @@ ELoadLocalizationResult LoadLocalization(const FString& LocalizationTargetName,
 		bool bHasPOFiles = false;
 		for (const FString& CultureToGenerate : OutLocTextHelper->GetAllCultures())
 		{
-			const FString POFilename = GetLocalizationTargetPOFile(OutLocTextHelper->GetTargetName(), LocalizationTargetDirectory, CultureToGenerate);;
-			if (FPaths::FileExists(POFilename))
+			if (const FString POFilename = GetLocalizationTargetPOFile(OutLocTextHelper->GetTargetName(), LocalizationTargetDirectory, CultureToGenerate);
+				FPaths::FileExists(POFilename))
 			{
 				bHasPOFiles = true;
 				break;
@@ -832,11 +849,19 @@ ELoadLocalizationResult LoadLocalization(const FString& LocalizationTargetName,
 		}
 	}
 
+	// If the native PO file is missing then we effectively have no localization data, but consider that case a warning
+	if (const FString POFilename = GetLocalizationTargetPOFile(OutLocTextHelper->GetTargetName(), LocalizationTargetDirectory, OutLocTextHelper->GetNativeCulture());
+		!FPaths::FileExists(POFilename))
+	{
+		UE_LOG(LogUGCLocalization, Warning, TEXT("Missing PO file for '%s' (culture '%s')"), *OutLocTextHelper->GetTargetName(), *OutLocTextHelper->GetNativeCulture());
+		return ELoadLocalizationResult::NoData;
+	}
+
 	// Import each PO file data, as we'll use it to generate the LocRes (via FLocTextHelper)
 	// We always process the native culture first as it's also used to populate the manifest with the source texts
 	if (!ImportPortableObject(LocalizationTargetDirectory, OutLocTextHelper->GetNativeCulture(), UGCLocDescriptor.PoFormat, *OutLocTextHelper))
 	{
-		UE_LOG(LogUGCLocalization, Error, TEXT("Failed to load PO file for '%s' (culture '%s')"), *LocalizationTargetName, *OutLocTextHelper->GetNativeCulture());
+		UE_LOG(LogUGCLocalization, Error, TEXT("Failed to load PO file for '%s' (culture '%s')"), *OutLocTextHelper->GetTargetName(), *OutLocTextHelper->GetNativeCulture());
 		return ELoadLocalizationResult::Failed;
 	}
 	for (const FString& CultureToGenerate : OutLocTextHelper->GetAllCultures())
@@ -846,9 +871,17 @@ ELoadLocalizationResult LoadLocalization(const FString& LocalizationTargetName,
 			continue;
 		}
 
+		// If the foreign PO file is missing then consider that a warning rather than an error
+		if (const FString POFilename = GetLocalizationTargetPOFile(OutLocTextHelper->GetTargetName(), LocalizationTargetDirectory, CultureToGenerate);
+			!FPaths::FileExists(POFilename))
+		{
+			UE_LOG(LogUGCLocalization, Warning, TEXT("Missing PO file for '%s' (culture '%s')"), *OutLocTextHelper->GetTargetName(), *CultureToGenerate);
+			continue;
+		}
+
 		if (!ImportPortableObject(LocalizationTargetDirectory, CultureToGenerate, UGCLocDescriptor.PoFormat, *OutLocTextHelper))
 		{
-			UE_LOG(LogUGCLocalization, Error, TEXT("Failed to load PO file for '%s' (culture '%s')"), *LocalizationTargetName, *CultureToGenerate);
+			UE_LOG(LogUGCLocalization, Error, TEXT("Failed to load PO file for '%s' (culture '%s')"), *OutLocTextHelper->GetTargetName(), *CultureToGenerate);
 			return ELoadLocalizationResult::Failed;
 		}
 	}
