Switched some Logs to Display so that they appear in CIS

#rnx

diff --git a/Engine/Source/Developer/Localization/Private/UserGeneratedContentLocalization.cpp b/Engine/Source/Developer/Localization/Private/UserGeneratedContentLocalization.cpp
--- a/Engine/Source/Developer/Localization/Private/UserGeneratedContentLocalization.cpp
+++ b/Engine/Source/Developer/Localization/Private/UserGeneratedContentLocalization.cpp
@@ -317,7 +317,7 @@ bool ExportLocalization(TArrayView<const TSharedRef<IPlugin>> Plugins, const FEx
 			{
 				if (IFileManager::Get().Copy(*DestinationFilename, *SourceFilename) == COPY_OK)
 				{
-					UE_LOG(LogUGCLocalization, Log, TEXT("Imported existing file for '%s': %s"), *PluginLocalizationTargetName, *SourceFilename);
+					UE_LOG(LogUGCLocalization, Display, TEXT("Imported existing file for '%s': %s"), *PluginLocalizationTargetName, *SourceFilename);
 				}
 				else
 				{
@@ -593,7 +593,7 @@ bool ExportLocalization(TArrayView<const TSharedRef<IPlugin>> Plugins, const FEx
 					if (IFileManager::Get().Copy(*DestinationFilename, *SourceFilename) == COPY_OK)
 					{
 						PostWriteFileWithSCC(DestinationFilename, bUseSourceControl);
-						UE_LOG(LogUGCLocalization, Log, TEXT("Updated file for '%s': %s"), *PluginLocalizationTargetName, *DestinationFilename);
+						UE_LOG(LogUGCLocalization, Display, TEXT("Updated file for '%s': %s"), *PluginLocalizationTargetName, *DestinationFilename);
 					}
 					else
 					{
@@ -634,7 +634,7 @@ bool ExportLocalization(TArrayView<const TSharedRef<IPlugin>> Plugins, const FEx
 				if (Plugin->UpdateDescriptor(PluginDescriptor, DescriptorUpdateFailureReason))
 				{
 					PostWriteFileWithSCC(Plugin->GetDescriptorFileName(), bUseSourceControl);
-					UE_LOG(LogUGCLocalization, Log, TEXT("Updated .uplugin file for '%s'"), *PluginLocalizationTargetName);
+					UE_LOG(LogUGCLocalization, Display, TEXT("Updated .uplugin file for '%s'"), *PluginLocalizationTargetName);
 				}
 				else
 				{
