Demoted a localization WARNING to LOG because it would always happen when a new plugin is freshly created.

diff --git a/Engine/Source/Developer/Localization/Private/PortableObjectPipeline.cpp b/Engine/Source/Developer/Localization/Private/PortableObjectPipeline.cpp
--- a/Engine/Source/Developer/Localization/Private/PortableObjectPipeline.cpp
+++ b/Engine/Source/Developer/Localization/Private/PortableObjectPipeline.cpp
@@ -220,7 +220,7 @@ namespace
 
 		if (!FPaths::FileExists(InPOFilePath))
 		{
-			UE_LOG(LogPortableObjectPipeline, Warning, TEXT("Could not find file %s"), *InPOFilePath);
+			UE_LOG(LogPortableObjectPipeline, Log, TEXT("Could not find file %s"), *InPOFilePath);
 			return true; // We don't fail on a missing file as the automation pipeline will always import before an export for a new language
 		}
 
