Fixed a typo in a localization warning

diff --git a/Engine/Source/Developer/Localization/Private/TextLocalizationResourceGenerator.cpp b/Engine/Source/Developer/Localization/Private/TextLocalizationResourceGenerator.cpp
--- a/Engine/Source/Developer/Localization/Private/TextLocalizationResourceGenerator.cpp
+++ b/Engine/Source/Developer/Localization/Private/TextLocalizationResourceGenerator.cpp
@@ -202,7 +202,7 @@ bool FTextLocalizationResourceGenerator::GenerateLocRes(const FLocTextHelper& In
 				if (!bIsNativeCulture && EnumHasAnyFlags(InGenerateFlags, EGenerateLocResFlags::ValidateRichTextTags) &&
 					!UE::TextLocalizationResourceGenerator::Private::ValidateRichTextTags(InManifestEntry->Source.Text, TranslationText.Text))
 				{
-					UE_LOGFMT(LogTextLocalizationResourceGenerator, Warning, "Broken Rich Text Tag detected in a translation. An unbalanced tag (a complete/incomplet opening rich text tag (i.e. <TagName>) with an incomplet/complete closing tag(</>)) was detected a the translation but not in its source text. Find the problematic tag in the translation and fix the translation to remove this warning. Translation File:'{translationFile}' Namespace And Key:'{locNamespace},{locKey}' Translation Text To Fix:'{text}'.",
+					UE_LOGFMT(LogTextLocalizationResourceGenerator, Warning, "Broken Rich Text Tag detected in a translation. An unbalanced tag (a complete/incomplet opening rich text tag (i.e. <TagName>) with an incomplet/complete closing tag(</>)) was detected in a the translation but not in its source text. Find the problematic tag in the translation and fix the translation to remove this warning. Translation File:'{translationFile}' Namespace And Key:'{locNamespace},{locKey}' Translation Text To Fix:'{text}'.",
 						("cultureCode", *InCultureToGenerate),
 						("locNamespace", *InManifestEntry->Namespace.GetString()),
 						("locKey", *Context.Key.GetString()),
