Added missing config tag to properties in UUserGeneratedContentLocalizationSettings

#rnx

diff --git a/Engine/Source/Developer/Localization/Public/UserGeneratedContentLocalization.h b/Engine/Source/Developer/Localization/Public/UserGeneratedContentLocalization.h
--- a/Engine/Source/Developer/Localization/Public/UserGeneratedContentLocalization.h
+++ b/Engine/Source/Developer/Localization/Public/UserGeneratedContentLocalization.h
@@ -34,20 +34,20 @@ public:
 	 * List of cultures that should be disabled for UGC localization.
 	 * @note You can't disable the native culture for the project.
 	 */
-	UPROPERTY(EditAnywhere, Category=Localization)
+	UPROPERTY(config, EditAnywhere, Category=Localization)
 	TArray<FString> CulturesToDisable;
 
 	/**
 	 * Should we compile UGC localization (if present) for DLC plugins during cook?
 	 */
-	UPROPERTY(EditAnywhere, Category=Localization)
+	UPROPERTY(config, EditAnywhere, Category=Localization)
 	bool bCompileDLCLocalizationDuringCook = true;
 
 	/**
 	 * Should we validate UGC localization (if present) for DLC plugins during cook?
 	 * @note Validation will happen against a UGC localization descriptor that has had InitializeFromProject called on it.
 	 */
-	UPROPERTY(EditAnywhere, Category=Localization)
+	UPROPERTY(config, EditAnywhere, Category=Localization)
 	bool bValidateDLCLocalizationDuringCook = true;
 };
 
