Added a mechanism to customize the localization target name of UGC plugins.

diff --git a/Engine/Source/Developer/Localization/Private/UserGeneratedContentLocalization.cpp b/Engine/Source/Developer/Localization/Private/UserGeneratedContentLocalization.cpp
--- a/Engine/Source/Developer/Localization/Private/UserGeneratedContentLocalization.cpp
+++ b/Engine/Source/Developer/Localization/Private/UserGeneratedContentLocalization.cpp
@@ -137,27 +137,79 @@ FString GetLocalizationScratchDirectory()
 	return FPaths::Combine(FPaths::ProjectIntermediateDir(), TEXT("Localization"));
 }
 
-FString GetLocalizationScratchDirectory(const TSharedRef<IPlugin>& Plugin)
+FString GetLocalizationScratchDirectory(const FString& LocalizationTargetName)
 {
-	return FPaths::Combine(GetLocalizationScratchDirectory(), Plugin->GetName());
+	return FPaths::Combine(GetLocalizationScratchDirectory(), LocalizationTargetName);
+}
+
+FStringView GetLocalizationTargetNameFieldName()
+{
+	return TEXTVIEW("UGCLocalizationTargetName");
 }
 
 FString GetLocalizationTargetName(const TSharedRef<IPlugin>& Plugin)
 {
+	FString LocalizationTargetName;
+
+#if WITH_EDITOR
+	if (const TSharedPtr<FJsonObject>& CachedDescriptorJson = Plugin->GetDescriptorJson())
+	{
+		CachedDescriptorJson->TryGetStringField(GetLocalizationTargetNameFieldName(), LocalizationTargetName);
+	}
+	else
+#endif // WITH_EDITOR
+	{
+		FString FileContents;
+		if (FFileHelper::LoadFileToString(FileContents, *Plugin->GetDescriptorFileName()))
+		{
+			TSharedRef<TJsonReader<>> Reader = TJsonReaderFactory<>::Create(FileContents);
+			TSharedPtr<FJsonObject> DescriptorJson;
+			if (FJsonSerializer::Deserialize(Reader, DescriptorJson) && DescriptorJson.IsValid())
+			{
+				DescriptorJson->TryGetStringField(GetLocalizationTargetNameFieldName(), LocalizationTargetName);
+			}
+		}
+	}
+
+	if (!LocalizationTargetName.IsEmpty())
+	{
+		return LocalizationTargetName;
+	}
+
 	// Note: If you change this naming scheme you'll need to handle backwards computability with existing data
 	return Plugin->GetName();
 }
 
+#if WITH_EDITOR
+
+bool SetLocalizationTargetName(const TSharedRef<IPlugin>& Plugin, const FString& LocalizationTargetName, bool bUseSourceControl, FText& OutFailReason)
+{
+	FPluginDescriptor PluginDescriptor = Plugin->GetDescriptor();
+	if (!LocalizationTargetName.IsEmpty())
+	{
+		PluginDescriptor.AdditionalFieldsToWrite.Emplace(GetLocalizationTargetNameFieldName(), MakeShared<FJsonValueString>(LocalizationTargetName));
+	}
+	else
+	{
+		PluginDescriptor.AdditionalFieldsToRemove.Emplace(GetLocalizationTargetNameFieldName());
+	}
+
+	PreWriteFileWithSCC(Plugin->GetDescriptorFileName(), bUseSourceControl);
+	const bool bResult = Plugin->UpdateDescriptor(PluginDescriptor, OutFailReason);
+	PostWriteFileWithSCC(Plugin->GetDescriptorFileName(), bUseSourceControl);
+	return bResult;
+}
+
+#endif // WITH_EDITOR
+
 FString GetLocalizationTargetDirectory(const TSharedRef<IPlugin>& Plugin)
 {
-	return FPaths::Combine(Plugin->GetContentDir(), TEXT("Localization"), Plugin->GetName());
+	return GetLocalizationTargetDirectory(GetLocalizationTargetName(Plugin), Plugin->GetContentDir());
 }
 
 FString GetLocalizationTargetDirectory(const TSharedRef<IPlugin>& Plugin, const FString& LocalizationTargetRootDirectoryOverride)
 {
-	return LocalizationTargetRootDirectoryOverride.IsEmpty()
-		? GetLocalizationTargetDirectory(Plugin)
-		: FPaths::Combine(LocalizationTargetRootDirectoryOverride, Plugin->GetName());
+	return GetLocalizationTargetDirectory(GetLocalizationTargetName(Plugin), Plugin->GetContentDir(), LocalizationTargetRootDirectoryOverride);
 }
 
 FString GetLocalizationTargetDirectory(const FString& LocalizationTargetName, const FString& PluginContentDirectory)
@@ -165,6 +217,13 @@ FString GetLocalizationTargetDirectory(const FString& LocalizationTargetName, co
 	return FPaths::Combine(PluginContentDirectory, TEXT("Localization"), LocalizationTargetName);
 }
 
+FString GetLocalizationTargetDirectory(const FString& LocalizationTargetName, const FString& PluginContentDirectory, const FString& LocalizationTargetRootDirectoryOverride)
+{
+	return LocalizationTargetRootDirectoryOverride.IsEmpty()
+		? GetLocalizationTargetDirectory(LocalizationTargetName, PluginContentDirectory)
+		: FPaths::Combine(LocalizationTargetRootDirectoryOverride, LocalizationTargetName);
+}
+
 FString GetLocalizationTargetUGCLocFile(const FString& LocalizationTargetName, const FString& LocalizationTargetDirectory)
 {
 	return FPaths::Combine(LocalizationTargetDirectory, LocalizationTargetName + TEXT(".ugcloc"));
@@ -299,8 +358,8 @@ bool ExportLocalization(TArrayView<const TSharedRef<IPlugin>> Plugins, const FEx
 	for (const TSharedRef<IPlugin>& Plugin : Plugins)
 	{
 		const FString PluginLocalizationTargetName = GetLocalizationTargetName(Plugin);
-		const FString PluginLocalizationScratchDirectory = GetLocalizationScratchDirectory(Plugin);
-		const FString PluginLocalizationTargetDirectory = GetLocalizationTargetDirectory(Plugin, ExportOptions.LocalizationTargetRootDirectoryOverride);
+		const FString PluginLocalizationScratchDirectory = GetLocalizationScratchDirectory(PluginLocalizationTargetName);
+		const FString PluginLocalizationTargetDirectory = GetLocalizationTargetDirectory(PluginLocalizationTargetName, Plugin->GetContentDir(), ExportOptions.LocalizationTargetRootDirectoryOverride);
 
 		// Track the source file hashes when the export started, so that we can detect post-export whether the files have actually changed
 		TMap<FString, FXxHash64>& LocalizationFileHashes = PerTargetLocalizationFileHashes.FindOrAdd(PluginLocalizationTargetName);
@@ -326,7 +385,7 @@ bool ExportLocalization(TArrayView<const TSharedRef<IPlugin>> Plugins, const FEx
 		}
 		if (!ExportOptions.LocalizationTargetRootDirectoryOverride.IsEmpty() && ExportOptions.MergeProjectDataWithRootDirectoryOverrideData.IsSet())
 		{
-			const FString PluginLocalizationSourceDirectory = GetLocalizationTargetDirectory(Plugin);
+			const FString PluginLocalizationSourceDirectory = GetLocalizationTargetDirectory(PluginLocalizationTargetName, Plugin->GetContentDir());
 			if (!MergeLocalization(PluginLocalizationTargetName, PluginLocalizationSourceDirectory, PluginLocalizationTargetName, PluginLocalizationScratchDirectory, ExportOptions.MergeProjectDataWithRootDirectoryOverrideData.GetValue(), /*bUseSourceControl*/false))
 			{
 				return false;
@@ -560,8 +619,8 @@ bool ExportLocalization(TArrayView<const TSharedRef<IPlugin>> Plugins, const FEx
 	for (const TSharedRef<IPlugin>& Plugin : Plugins)
 	{
 		const FString PluginLocalizationTargetName = GetLocalizationTargetName(Plugin);
-		const FString PluginLocalizationScratchDirectory = GetLocalizationScratchDirectory(Plugin);
-		const FString PluginLocalizationTargetDirectory = GetLocalizationTargetDirectory(Plugin, ExportOptions.LocalizationTargetRootDirectoryOverride);
+		const FString PluginLocalizationScratchDirectory = GetLocalizationScratchDirectory(PluginLocalizationTargetName);
+		const FString PluginLocalizationTargetDirectory = GetLocalizationTargetDirectory(PluginLocalizationTargetName, Plugin->GetContentDir(), ExportOptions.LocalizationTargetRootDirectoryOverride);
 
 		const bool bUseSourceControl = ExportOptions.LocalizationTargetRootDirectoryOverride.IsEmpty();
 
@@ -707,7 +766,7 @@ bool CompileLocalization(TArrayView<const TSharedRef<IPlugin>> Plugins, const FU
 	for (const TSharedRef<IPlugin>& Plugin : Plugins)
 	{
 		const FString PluginLocalizationTargetName = GetLocalizationTargetName(Plugin);
-		const FString PluginLocalizationTargetDirectory = GetLocalizationTargetDirectory(Plugin);
+		const FString PluginLocalizationTargetDirectory = GetLocalizationTargetDirectory(PluginLocalizationTargetName, Plugin->GetContentDir());
 		if (!CompileLocalization(PluginLocalizationTargetName, PluginLocalizationTargetDirectory, PluginLocalizationTargetDirectory, DefaultDescriptor))
 		{
 			return false;
@@ -1085,7 +1144,7 @@ void CleanupLocalization(TArrayView<const TSharedRef<IPlugin>> PluginsToClean, T
 		for (const TSharedRef<IPlugin>& Plugin : Plugins)
 		{
 			const FString PluginLocalizationTargetName = GetLocalizationTargetName(Plugin);
-			const FString PluginLocalizationTargetDirectory = GetLocalizationTargetDirectory(Plugin);
+			const FString PluginLocalizationTargetDirectory = GetLocalizationTargetDirectory(PluginLocalizationTargetName, Plugin->GetContentDir());
 
 			// Find any leftover PO files to cleanup
 			const FString PluginPOFilename = PluginLocalizationTargetName + TEXT(".po");
diff --git a/Engine/Source/Developer/Localization/Public/UserGeneratedContentLocalization.h b/Engine/Source/Developer/Localization/Public/UserGeneratedContentLocalization.h
--- a/Engine/Source/Developer/Localization/Public/UserGeneratedContentLocalization.h
+++ b/Engine/Source/Developer/Localization/Public/UserGeneratedContentLocalization.h
@@ -183,11 +183,25 @@ enum class EMergeLocalizationMode : uint8
 	All,
 };
 
+/**
+ * Gets the name of the custom plugin field used to store the localization target name.
+ */
+LOCALIZATION_API FStringView GetLocalizationTargetNameFieldName();
+
 /**
  * Utility to get the name of the UGC localization target that would be used by the given plugin.
  */
 LOCALIZATION_API FString GetLocalizationTargetName(const TSharedRef<IPlugin>& Plugin);
 
+#if WITH_EDITOR
+
+/**
+ * Utility to set the name of the UGC localization target that will be used by the given plugin.
+ */
+LOCALIZATION_API bool SetLocalizationTargetName(const TSharedRef<IPlugin>& Plugin, const FString& LocalizationTargetName, bool bUseSourceControl, FText& OutFailReason);
+
+#endif // WITH_EDITOR
+
 /**
  * Utility to get the path of the UGC localization target directory that would be used by the given plugin.
  */
@@ -199,10 +213,15 @@ LOCALIZATION_API FString GetLocalizationTargetDirectory(const TSharedRef<IPlugin
 LOCALIZATION_API FString GetLocalizationTargetDirectory(const TSharedRef<IPlugin>& Plugin, const FString& LocalizationTargetRootDirectoryOverride);
 
 /**
- * Utility to get the path of the UGC localization target that would be used by the given plugin name and directory.
+ * Utility to get the path of the UGC localization target that would be used by the given target name and directory.
  */
 LOCALIZATION_API FString GetLocalizationTargetDirectory(const FString& LocalizationTargetName, const FString& PluginContentDirectory);
 
+/**
+ * Utility to get the path of the UGC localization target that would be used by the given target name and directory, based on the given root directory override (@see FExportLocalizationOptions).
+ */
+LOCALIZATION_API FString GetLocalizationTargetDirectory(const FString& LocalizationTargetName, const FString& PluginContentDirectory, const FString& LocalizationTargetRootDirectoryOverride);
+
 /**
  * Utility to get the UGCLoc file for for the given UGC localization target name and directory.
  */
