Improved manifest sort determinism when an ID+source pair has multiple source locations

We now sort by the source location as a secondary sort in cases where two contexts have the same key

#rnx

diff --git a/Engine/Source/Developer/Localization/Private/Serialization/JsonInternationalizationManifestSerializer.cpp b/Engine/Source/Developer/Localization/Private/Serialization/JsonInternationalizationManifestSerializer.cpp
--- a/Engine/Source/Developer/Localization/Private/Serialization/JsonInternationalizationManifestSerializer.cpp
+++ b/Engine/Source/Developer/Localization/Private/Serialization/JsonInternationalizationManifestSerializer.cpp
@@ -28,24 +28,6 @@ const FString FJsonInternationalizationManifestSerializer::TAG_METADATA_KEY = TE
 const FString FJsonInternationalizationManifestSerializer::NAMESPACE_DELIMITER = TEXT(".");
 
 
-struct FCompareManifestEntryBySource
-{
-	FORCEINLINE bool operator()( TSharedPtr< FManifestEntry > A, TSharedPtr< FManifestEntry > B ) const
-	{
-		return A->Source < B->Source;
-	}
-};
-
-
-struct FCompareStructuredEntryByNamespace
-{
-	FORCEINLINE bool operator()( TSharedPtr< FStructuredEntry > A, TSharedPtr< FStructuredEntry > B ) const
-	{
-		return A->Namespace < B->Namespace;
-	}
-};
-
-
 bool FJsonInternationalizationManifestSerializer::DeserializeManifest( const FString& InStr, TSharedRef< FInternationalizationManifest > Manifest, const FName PlatformName )
 {
 	TSharedPtr<FJsonObject> JsonManifestObj;
@@ -399,20 +381,36 @@ void FJsonInternationalizationManifestSerializer::SortStructuredData( TSharedPtr
 	}
 
 	// Sort the manifest entries by source text.
-	InElement->ManifestEntries.Sort( FCompareManifestEntryBySource() );
+	InElement->ManifestEntries.Sort(
+		[](const TSharedPtr< FManifestEntry >& A, const TSharedPtr< FManifestEntry >& B)
+		{
+			return A->Source < B->Source;
+		});
 
-	// Sort the manifest entry contexts by key/identifier
+	// Sort the manifest entry contexts by key/identifier (primary) and source location (secondary)
 	for( TArray< TSharedPtr< FManifestEntry > >::TIterator Iter( InElement->ManifestEntries.CreateIterator() ); Iter; ++Iter)
 	{
 		TSharedPtr< FManifestEntry > SubEntry = *Iter;
 		if( SubEntry.IsValid())
 		{
-			SubEntry->Contexts.Sort( );
+			SubEntry->Contexts.Sort(
+				[](const FManifestContext& A, const FManifestContext& B)
+				{
+					if (A == B)
+					{
+						return A.SourceLocation < B.SourceLocation;
+					}
+					return A < B;
+				});
 		}
 	}
 	
 	// Sort the subnamespaces by namespace string
-	InElement->SubNamespaces.Sort( FCompareStructuredEntryByNamespace() );
+	InElement->SubNamespaces.Sort(
+		[](const TSharedPtr< FStructuredEntry >& A, const TSharedPtr< FStructuredEntry >& B)
+		{
+			return A->Namespace < B->Namespace;
+		});
 
 	// Do the sorting for each of the subnamespaces
 	for( TArray< TSharedPtr< FStructuredEntry > >::TIterator Iter( InElement->SubNamespaces.CreateIterator() ); Iter; ++Iter )
