Optimized manifest/archive generation and sorting during export

* FJsonInternationalizationManifestSerializer/FJsonInternationalizationArchiveSerializer::GenerateStructuredData:
* Replaced linear look-up of sub-namespaces with a map.
* Replaced AddUnique with Add, for a case where we know the entry we're adding is already unique.
* FJsonInternationalizationManifestSerializer/FJsonInternationalizationArchiveSerializer::SortStructuredData:
* Fixed determinism issues caused by case-insensitive compares, and archives not sorting on their key values.

For an internal workflow, this reduced the combined manifest and archive export time from ~200s to ~60s.

diff --git a/Engine/Source/Developer/Localization/Private/Serialization/JsonInternationalizationArchiveSerializer.cpp b/Engine/Source/Developer/Localization/Private/Serialization/JsonInternationalizationArchiveSerializer.cpp
--- a/Engine/Source/Developer/Localization/Private/Serialization/JsonInternationalizationArchiveSerializer.cpp
+++ b/Engine/Source/Developer/Localization/Private/Serialization/JsonInternationalizationArchiveSerializer.cpp
@@ -28,40 +28,6 @@ const FString FJsonInternationalizationArchiveSerializer::TAG_METADATA_KEY = TEX
 const FString FJsonInternationalizationArchiveSerializer::NAMESPACE_DELIMITER = TEXT(".");
 
 
-struct FCompareArchiveEntryBySourceAndKey
-{
-	FORCEINLINE bool operator()( TSharedPtr< FArchiveEntry > A, TSharedPtr< FArchiveEntry > B ) const
-	{
-		bool bResult = false;
-		if( A->Source < B->Source )
-		{
-			bResult = true;
-		}
-		else if( A->Source == B->Source )
-		{
-			if( A->KeyMetadataObj.IsValid() != B->KeyMetadataObj.IsValid() )
-			{
-				bResult = B->KeyMetadataObj.IsValid();
-			}
-			else if( A->KeyMetadataObj.IsValid() && B->KeyMetadataObj.IsValid() )
-			{
-				bResult = (*(A->KeyMetadataObj) < *(B->KeyMetadataObj));
-			}
-		}
-		return bResult;
-	}
-};
-
-
-struct FCompareStructuredArchiveEntryByNamespace
-{
-	FORCEINLINE bool operator()( TSharedPtr< FStructuredArchiveEntry > A, TSharedPtr< FStructuredArchiveEntry > B ) const
-	{
-		return A->Namespace < B->Namespace;
-	}
-};
-
-
 bool FJsonInternationalizationArchiveSerializer::DeserializeArchive(const FString& InStr, TSharedRef<FInternationalizationArchive> InArchive, TSharedPtr<const FInternationalizationManifest> InManifest, TSharedPtr<const FInternationalizationArchive> InNativeArchive)
 {
 	TSharedPtr< FJsonObject > JsonArchiveObj;
@@ -391,41 +357,32 @@ bool FJsonInternationalizationArchiveSerializer::JsonObjToArchive(TSharedRef<FJs
 
 void FJsonInternationalizationArchiveSerializer::GenerateStructuredData( TSharedRef< const FInternationalizationArchive > InArchive, TSharedPtr<FStructuredArchiveEntry> RootElement )
 {
-	//Loop through all the unstructured archive entries and build up our structured hierarchy
-	for(FArchiveEntryByStringContainer::TConstIterator It( InArchive->GetEntriesBySourceTextIterator() ); It; ++It)
+	// Loop through all the unstructured archive entries and build up our structured hierarchy
+	TArray<FString> NamespaceTokens;
+	for (FArchiveEntryByStringContainer::TConstIterator It(InArchive->GetEntriesBySourceTextIterator()); It; ++It)
 	{
-		const TSharedRef< FArchiveEntry > UnstructuredArchiveEntry = It.Value();
-
-		TArray< FString > NamespaceTokens;
+		const TSharedRef<FArchiveEntry> UnstructuredArchiveEntry = It.Value();
 
 		// Tokenize the namespace by using '.' as a delimiter
-		int32 NamespaceTokenCount = UnstructuredArchiveEntry->Namespace.GetString().ParseIntoArray( NamespaceTokens, *NAMESPACE_DELIMITER, true );
+		NamespaceTokens.Reset();
+		UnstructuredArchiveEntry->Namespace.GetString().ParseIntoArray(NamespaceTokens, *NAMESPACE_DELIMITER, true);
 
-		TSharedPtr< FStructuredArchiveEntry > StructuredArchiveEntry = RootElement;
-		//Loop through all the namespace tokens and find the appropriate structured entry, if it does not exist add it.  At the end StructuredArchiveEntry
-		//  will point to the correct hierarchy entry for a given namespace
-		for( int32 TokenIndex = 0; TokenIndex < NamespaceTokenCount; ++TokenIndex )
-		{
-			TSharedPtr<FStructuredArchiveEntry> FoundNamespaceEntry;
-			for( int SubNamespaceIndex = 0; SubNamespaceIndex < StructuredArchiveEntry->SubNamespaces.Num(); SubNamespaceIndex++ )
-			{
-				if( StructuredArchiveEntry->SubNamespaces[SubNamespaceIndex]->Namespace.Equals(NamespaceTokens[TokenIndex], ESearchCase::CaseSensitive) )
-				{
-					FoundNamespaceEntry = StructuredArchiveEntry->SubNamespaces[SubNamespaceIndex];
-					break;
-				}
-			}
+		TSharedPtr<FStructuredArchiveEntry> StructuredArchiveEntry = RootElement;
 
-			if( !FoundNamespaceEntry.IsValid() )
+		// Loop through all the namespace tokens and find the appropriate structured entry, if it does not exist add it.
+		// At the end StructuredArchiveEntry will point to the correct hierarchy entry for a given namespace
+		for (const FString& NamespaceToken : NamespaceTokens)
+		{
+			TSharedPtr<FStructuredArchiveEntry> FoundNamespaceEntry = StructuredArchiveEntry->SubNamespaces.FindRef(NamespaceToken);
+			if (!FoundNamespaceEntry)
 			{
-				int32 index = StructuredArchiveEntry->SubNamespaces.Add( MakeShareable( new FStructuredArchiveEntry( NamespaceTokens[TokenIndex] ) ) );
-				FoundNamespaceEntry = StructuredArchiveEntry->SubNamespaces[index];
+				FoundNamespaceEntry = StructuredArchiveEntry->SubNamespaces.Add(NamespaceToken, MakeShared<FStructuredArchiveEntry>(NamespaceToken));
 			}
-			StructuredArchiveEntry = FoundNamespaceEntry;
+			StructuredArchiveEntry = MoveTemp(FoundNamespaceEntry);
 		}
 
 		// We add the unstructured Archive entry to the hierarchy
-		StructuredArchiveEntry->ArchiveEntries.AddUnique( UnstructuredArchiveEntry );
+		StructuredArchiveEntry->ArchiveEntries.Add(UnstructuredArchiveEntry);
 	}
 }
 
@@ -437,16 +394,39 @@ void FJsonInternationalizationArchiveSerializer::SortStructuredData( TSharedPtr<
 		return;
 	}
 
-	// Sort the manifest entries by source text.
-	InElement->ArchiveEntries.Sort( FCompareArchiveEntryBySourceAndKey() );
+	// Sort the manifest entries by source text (primary) and key (secondary).
+	InElement->ArchiveEntries.Sort(
+		[](const TSharedPtr<FArchiveEntry>& A, const TSharedPtr<FArchiveEntry>& B)
+		{
+			if (A->Source == B->Source)
+			{
+				if (A->Key == B->Key)
+				{
+					if (A->KeyMetadataObj.IsValid() != B->KeyMetadataObj.IsValid())
+					{
+						return B->KeyMetadataObj.IsValid();
+					}
+					if (A->KeyMetadataObj.IsValid() && B->KeyMetadataObj.IsValid())
+					{
+						return (*A->KeyMetadataObj < *B->KeyMetadataObj);
+					}
+				}
+				return A->Key < B->Key;
+			}
+			return A->Source < B->Source;
+		});
 
 	// Sort the subnamespaces by namespace string
-	InElement->SubNamespaces.Sort( FCompareStructuredArchiveEntryByNamespace() );
+	InElement->SubNamespaces.KeySort(
+		[](const FString& A, const FString& B)
+		{
+			return A.Compare(B, ESearchCase::CaseSensitive) < 0;
+		});
 
 	// Do the sorting for each of the subnamespaces
-	for( TArray< TSharedPtr< FStructuredArchiveEntry > >::TIterator Iter( InElement->SubNamespaces.CreateIterator() ); Iter; ++Iter )
+	for( auto Iter = InElement->SubNamespaces.CreateIterator(); Iter; ++Iter )
 	{
-		TSharedPtr< FStructuredArchiveEntry > SubElement = *Iter;
+		TSharedPtr< FStructuredArchiveEntry > SubElement = Iter->Value;
 
 		SortStructuredData( SubElement );
 	}
@@ -522,9 +502,9 @@ void FJsonInternationalizationArchiveSerializer::StructuredDataToJsonObj( TShare
 	}
 
 	//Write the subnamespaces
-	for( TArray< TSharedPtr< FStructuredArchiveEntry > >::TConstIterator Iter( InElement->SubNamespaces ); Iter; ++Iter )
+	for( auto Iter = InElement->SubNamespaces.CreateConstIterator(); Iter; ++Iter )
 	{
-		const TSharedPtr<FStructuredArchiveEntry> SubElement = *Iter;
+		const TSharedPtr<FStructuredArchiveEntry> SubElement = Iter->Value;
 		if( SubElement.IsValid() )
 		{
 			TSharedRef<FJsonObject> SubObject = MakeShareable( new FJsonObject );
diff --git a/Engine/Source/Developer/Localization/Private/Serialization/JsonInternationalizationManifestSerializer.cpp b/Engine/Source/Developer/Localization/Private/Serialization/JsonInternationalizationManifestSerializer.cpp
--- a/Engine/Source/Developer/Localization/Private/Serialization/JsonInternationalizationManifestSerializer.cpp
+++ b/Engine/Source/Developer/Localization/Private/Serialization/JsonInternationalizationManifestSerializer.cpp
@@ -334,41 +334,32 @@ bool FJsonInternationalizationManifestSerializer::JsonObjToManifest( TSharedRef<
 
 void FJsonInternationalizationManifestSerializer::GenerateStructuredData( TSharedRef< const FInternationalizationManifest > InManifest, TSharedPtr< FStructuredEntry > RootElement )
 {
-	//Loop through all the unstructured manifest entries and build up our structured hierarchy
-	for( FManifestEntryByStringContainer::TConstIterator It( InManifest->GetEntriesBySourceTextIterator() ); It; ++It )
+	// Loop through all the unstructured manifest entries and build up our structured hierarchy
+	TArray<FString> NamespaceTokens;
+	for (FManifestEntryByStringContainer::TConstIterator It(InManifest->GetEntriesBySourceTextIterator()); It; ++It)
 	{
-		const TSharedRef< FManifestEntry > UnstructuredManifestEntry = It.Value();
-
-		TArray< FString > NamespaceTokens;
+		const TSharedRef<FManifestEntry> UnstructuredManifestEntry = It.Value();
 
 		// Tokenize the namespace by using '.' as a delimiter
-		int32 NamespaceTokenCount = UnstructuredManifestEntry->Namespace.GetString().ParseIntoArray( NamespaceTokens, *NAMESPACE_DELIMITER, true );
+		NamespaceTokens.Reset();
+		UnstructuredManifestEntry->Namespace.GetString().ParseIntoArray(NamespaceTokens, *NAMESPACE_DELIMITER, true);
 
-		TSharedPtr< FStructuredEntry > StructuredManifestEntry = RootElement;
-		//Loop through all the namespace tokens and find the appropriate structured entry, if it does not exist add it.  At the end StructuredManifestEntry
-		//  will point to the correct hierarchy entry for a given namespace
-		for( int32 TokenIndex = 0; TokenIndex < NamespaceTokenCount; ++TokenIndex )
-		{
-			TSharedPtr< FStructuredEntry > FoundNamespaceEntry;
-			for( int SubNamespaceIndex = 0; SubNamespaceIndex < StructuredManifestEntry->SubNamespaces.Num(); SubNamespaceIndex++ )
-			{
-				if( StructuredManifestEntry->SubNamespaces[SubNamespaceIndex]->Namespace.Equals(NamespaceTokens[TokenIndex], ESearchCase::CaseSensitive) )
-				{
-					FoundNamespaceEntry = StructuredManifestEntry->SubNamespaces[SubNamespaceIndex];
-					break;
-				}
-			}
+		TSharedPtr<FStructuredEntry> StructuredManifestEntry = RootElement;
 
-			if( !FoundNamespaceEntry.IsValid() )
+		// Loop through all the namespace tokens and find the appropriate structured entry, if it does not exist add it.
+		// At the end StructuredManifestEntry will point to the correct hierarchy entry for a given namespace
+		for (const FString& NamespaceToken : NamespaceTokens)
+		{
+			TSharedPtr<FStructuredEntry> FoundNamespaceEntry = StructuredManifestEntry->SubNamespaces.FindRef(NamespaceToken);
+			if (!FoundNamespaceEntry)
 			{
-				int32 index = StructuredManifestEntry->SubNamespaces.Add( MakeShareable( new FStructuredEntry( NamespaceTokens[TokenIndex] ) ) );
-				FoundNamespaceEntry = StructuredManifestEntry->SubNamespaces[index];
+				FoundNamespaceEntry = StructuredManifestEntry->SubNamespaces.Add(NamespaceToken, MakeShared<FStructuredEntry>(NamespaceToken));
 			}
-			StructuredManifestEntry = FoundNamespaceEntry;
+			StructuredManifestEntry = MoveTemp(FoundNamespaceEntry);
 		}
 
 		// We add the unstructured manifest entry to the hierarchy
-		StructuredManifestEntry->ManifestEntries.AddUnique( UnstructuredManifestEntry );
+		StructuredManifestEntry->ManifestEntries.Add(UnstructuredManifestEntry);
 	}
 }
 
@@ -398,7 +389,7 @@ void FJsonInternationalizationManifestSerializer::SortStructuredData( TSharedPtr
 				{
 					if (A == B)
 					{
-						return A.SourceLocation < B.SourceLocation;
+						return A.SourceLocation.Compare(B.SourceLocation, ESearchCase::CaseSensitive) < 0;
 					}
 					return A < B;
 				});
@@ -406,16 +397,16 @@ void FJsonInternationalizationManifestSerializer::SortStructuredData( TSharedPtr
 	}
 	
 	// Sort the subnamespaces by namespace string
-	InElement->SubNamespaces.Sort(
-		[](const TSharedPtr< FStructuredEntry >& A, const TSharedPtr< FStructuredEntry >& B)
+	InElement->SubNamespaces.KeySort(
+		[](const FString& A, const FString& B)
 		{
-			return A->Namespace < B->Namespace;
+			return A.Compare(B, ESearchCase::CaseSensitive) < 0;
 		});
 
 	// Do the sorting for each of the subnamespaces
-	for( TArray< TSharedPtr< FStructuredEntry > >::TIterator Iter( InElement->SubNamespaces.CreateIterator() ); Iter; ++Iter )
+	for( auto Iter = InElement->SubNamespaces.CreateIterator(); Iter; ++Iter )
 	{
-		TSharedPtr< FStructuredEntry > SubElement = *Iter;
+		TSharedPtr< FStructuredEntry > SubElement = Iter->Value;
 
 		SortStructuredData( SubElement );
 	}
@@ -512,9 +503,9 @@ void FJsonInternationalizationManifestSerializer::StructuredDataToJsonObj( TShar
 	}
 
 	//Write the subnamespaces
-	for( TArray< TSharedPtr< FStructuredEntry > >::TConstIterator Iter( InElement->SubNamespaces.CreateConstIterator() ); Iter; ++Iter )
+	for( auto Iter = InElement->SubNamespaces.CreateConstIterator(); Iter; ++Iter )
 	{
-		const TSharedPtr< FStructuredEntry > SubElement = *Iter;
+		const TSharedPtr< FStructuredEntry > SubElement = Iter->Value;
 
 		if(SubElement.IsValid())
 		{
diff --git a/Engine/Source/Developer/Localization/Public/Serialization/JsonInternationalizationArchiveSerializer.h b/Engine/Source/Developer/Localization/Public/Serialization/JsonInternationalizationArchiveSerializer.h
--- a/Engine/Source/Developer/Localization/Public/Serialization/JsonInternationalizationArchiveSerializer.h
+++ b/Engine/Source/Developer/Localization/Public/Serialization/JsonInternationalizationArchiveSerializer.h
@@ -3,35 +3,22 @@
 #pragma once
 
 #include "Containers/Array.h"
+#include "Containers/Map.h"
 #include "Containers/UnrealString.h"
 #include "CoreMinimal.h"
 #include "Dom/JsonObject.h"
 #include "Internationalization/InternationalizationArchive.h"
 #include "Internationalization/InternationalizationManifest.h"
+#include "Internationalization/LocKeyFuncs.h"
 #include "Templates/SharedPointer.h"
 
 #define UE_API LOCALIZATION_API
 
+class FArchiveEntry;
 class FInternationalizationArchive;
 class FInternationalizationManifest;
 class FJsonObject;
 
-/**
- * Used to arrange Internationalization archive data in a hierarchy based on namespace prior to json serialization.
- */
-struct FStructuredArchiveEntry
-{
-public:
-	FStructuredArchiveEntry( const FString& InNamespace )
-		: Namespace( InNamespace )
-	{ }
-
-	const FString Namespace;
-	TArray< TSharedPtr< FStructuredArchiveEntry > > SubNamespaces;
-	TArray< TSharedPtr< class FArchiveEntry > > ArchiveEntries;
-};
-
-
 /**
  * Implements a serializer that serializes to and from Json encoded data.
  */
@@ -107,6 +94,21 @@ public:
 
 protected:
 
+	/**
+	 * Used to arrange Internationalization archive data in a hierarchy based on namespace prior to json serialization.
+	 */
+	struct FStructuredArchiveEntry
+	{
+		explicit FStructuredArchiveEntry(const FString& InNamespace)
+			: Namespace(InNamespace)
+		{
+		}
+
+		const FString Namespace;
+		TMap<FString, TSharedPtr<FStructuredArchiveEntry>, FDefaultSetAllocator, FLocKeyMapFuncs<TSharedPtr<FStructuredArchiveEntry>>> SubNamespaces;
+		TArray<TSharedPtr<FArchiveEntry>> ArchiveEntries;
+	};
+
 	/**
 	 * Deserializes an archive from a JSON object.
 	 *
diff --git a/Engine/Source/Developer/Localization/Public/Serialization/JsonInternationalizationManifestSerializer.h b/Engine/Source/Developer/Localization/Public/Serialization/JsonInternationalizationManifestSerializer.h
--- a/Engine/Source/Developer/Localization/Public/Serialization/JsonInternationalizationManifestSerializer.h
+++ b/Engine/Source/Developer/Localization/Public/Serialization/JsonInternationalizationManifestSerializer.h
@@ -3,35 +3,24 @@
 #pragma once
 
 #include "Containers/Array.h"
+#include "Containers/Map.h"
 #include "Containers/UnrealString.h"
 #include "CoreMinimal.h"
 #include "Dom/JsonObject.h"
 #include "Internationalization/InternationalizationManifest.h"
+#include "Internationalization/LocKeyFuncs.h"
 #include "Templates/SharedPointer.h"
 #include "UObject/NameTypes.h"
 
 #define UE_API LOCALIZATION_API
 
 class FInternationalizationManifest;
-
-struct FStructuredEntry
-{
-public:
-	FStructuredEntry( const FString& InNamespace )
-		: Namespace( InNamespace )
-		, SubNamespaces()
-	{
-	}
-
-	const FString Namespace;
-	TArray< TSharedPtr< FStructuredEntry > > SubNamespaces;
-	TArray< TSharedPtr< class FManifestEntry > > ManifestEntries;
-};
-
-
+class FManifestEntry;
 class FJsonObject;
 
-
+/**
+ * Implements a serializer that serializes to and from Json encoded data.
+ */
 class FJsonInternationalizationManifestSerializer
 {
 public:
@@ -92,6 +81,18 @@ public:
 
 protected:
 
+	struct FStructuredEntry
+	{
+		explicit FStructuredEntry(const FString& InNamespace)
+			: Namespace(InNamespace)
+		{
+		}
+
+		const FString Namespace;
+		TMap<FString, TSharedPtr<FStructuredEntry>, FDefaultSetAllocator, FLocKeyMapFuncs<TSharedPtr<FStructuredEntry>>> SubNamespaces;
+		TArray<TSharedPtr<FManifestEntry>> ManifestEntries;
+	};
+
 	/**
 	 * Convert a JSON object to a Internationalization manifest.
 	 *
