Updated manifest/archive serialization to use rapidJSON

For an internal workflow, this reduced the FLocTextHelper::LoadAll time from ~50s to ~22s, and the FLocTextHelper::SaveAll time from ~59s to ~30s.

Breaking Changes:
* The overloads of FJsonInternationalizationManifestSerializer::SerializeManifest/DeserializeManifest and FJsonInternationalizationArchiveSerializer::SerializeArchive/DeserializeArchive using a FJsonObject have been removed, as serialization no longer uses FJsonObject internally, and these overloads were unused.
* The overloads of FJsonInternationalizationManifestSerializer::SerializeManifest/DeserializeManifest and FJsonInternationalizationArchiveSerializer::SerializeArchive/DeserializeArchive that use FString have been renamed to FJsonInternationalizationManifestSerializer::SerializeManifestToString/DeserializeManifestFromString and FJsonInternationalizationArchiveSerializer::SerializeArchiveToString/DeserializeArchiveFromFile, for consistency with the XToFile/XFromFile overloads.
* The API of FJsonInternationalizationMetaDataSerializer that was using FJsonObject has been replaced with rapidJSON types (via UE::Json), as serialization no longer uses FJsonObject internally, and these functions are effectively internal to FJsonInternationalizationManifestSerializer and FJsonInternationalizationArchiveSerializer.
* FLocTextHelper::SerializeManifestToJson has been removed, as it relied on the FJsonObject API in FJsonInternationalizationManifestSerializer, and was unused.

diff --git a/Engine/Source/Developer/Localization/Private/LocTextHelper.cpp b/Engine/Source/Developer/Localization/Private/LocTextHelper.cpp
--- a/Engine/Source/Developer/Localization/Private/LocTextHelper.cpp
+++ b/Engine/Source/Developer/Localization/Private/LocTextHelper.cpp
@@ -602,11 +602,6 @@ bool FLocTextHelper::SaveManifest(const FString& InManifestFilePath, FText* OutE
 	return SaveManifestImpl(Manifest.ToSharedRef(), InManifestFilePath, OutError);
 }
 
-bool FLocTextHelper::SerializeManifestToJson(TSharedRef<FJsonObject> JsonObject)
-{
-	return FJsonInternationalizationManifestSerializer::SerializeManifest(Manifest.ToSharedRef(), JsonObject);
-}
-
 void FLocTextHelper::TrimManifest()
 {
 	if (Dependencies.Num() > 0)
diff --git a/Engine/Source/Developer/Localization/Private/Serialization/JsonInternationalizationArchiveSerializer.cpp b/Engine/Source/Developer/Localization/Private/Serialization/JsonInternationalizationArchiveSerializer.cpp
--- a/Engine/Source/Developer/Localization/Private/Serialization/JsonInternationalizationArchiveSerializer.cpp
+++ b/Engine/Source/Developer/Localization/Private/Serialization/JsonInternationalizationArchiveSerializer.cpp
@@ -1,11 +1,10 @@
 // Copyright Epic Games, Inc. All Rights Reserved.
 
 #include "Serialization/JsonInternationalizationArchiveSerializer.h"
+
+#include "LocTextHelper.h"
 #include "Misc/FileHelper.h"
-#include "Serialization/JsonReader.h"
-#include "Serialization/JsonSerializer.h"
 #include "Serialization/JsonInternationalizationMetadataSerializer.h"
-#include "LocTextHelper.h"
 
 
 DEFINE_LOG_CATEGORY_STATIC(LogInternationalizationArchiveSerializer, Log, All);
@@ -28,23 +27,24 @@ const FString FJsonInternationalizationArchiveSerializer::TAG_METADATA_KEY = TEX
 const FString FJsonInternationalizationArchiveSerializer::NAMESPACE_DELIMITER = TEXT(".");
 
 
-bool FJsonInternationalizationArchiveSerializer::DeserializeArchive(const FString& InStr, TSharedRef<FInternationalizationArchive> InArchive, TSharedPtr<const FInternationalizationManifest> InManifest, TSharedPtr<const FInternationalizationArchive> InNativeArchive)
+bool FJsonInternationalizationArchiveSerializer::DeserializeArchiveFromString(const FString& InStr, TSharedRef<FInternationalizationArchive> InArchive, TSharedPtr<const FInternationalizationManifest> InManifest, TSharedPtr<const FInternationalizationArchive> InNativeArchive)
 {
-	TSharedPtr< FJsonObject > JsonArchiveObj;
-	TSharedRef< TJsonReader<> > Reader = TJsonReaderFactory<>::Create(InStr);
-	
-	if (FJsonSerializer::Deserialize(Reader, JsonArchiveObj) && JsonArchiveObj.IsValid())
+	TValueOrError<UE::Json::FDocument, UE::Json::FParseError> Document = UE::Json::Parse(InStr);
+
+	if (Document.HasError())
 	{
-		return DeserializeInternal(JsonArchiveObj.ToSharedRef(), InArchive, InManifest, InNativeArchive);
+		UE_LOG(LogInternationalizationArchiveSerializer, Error, TEXT("Failed to parse archive. %s."), *Document.GetError().CreateMessage(InStr));
+		return false;
 	}
 
-	return false;
-}
-
+	TOptional<UE::Json::FConstObject> RootObject = UE::Json::GetRootObject(Document.GetValue());
+	if (!RootObject.IsSet())
+	{
+		UE_LOG(LogInternationalizationArchiveSerializer, Error, TEXT("Failed to parse archive. No root object."));
+		return false;
+	}
 
-bool FJsonInternationalizationArchiveSerializer::DeserializeArchive(TSharedRef<FJsonObject> InJsonObj, TSharedRef<FInternationalizationArchive> InArchive, TSharedPtr<const FInternationalizationManifest> InManifest, TSharedPtr<const FInternationalizationArchive> InNativeArchive)
-{
-	return DeserializeInternal(InJsonObj, InArchive, InManifest, InNativeArchive);
+	return DeserializeInternal(*RootObject, InArchive, InManifest, InNativeArchive);
 }
 
 
@@ -58,58 +58,55 @@ bool FJsonInternationalizationArchiveSerializer::DeserializeArchiveFromFile(cons
 		return false;
 	}
 
-	// Parse as JSON
-	TSharedPtr<FJsonObject> JsonObject;
+	// Grab the internal character array so we can do an insitu parse
+	TArray<TCHAR> FileData = MoveTemp(FileContents.GetCharArray());
+	TValueOrError<UE::Json::FDocument, UE::Json::FParseError> Document = UE::Json::ParseInPlace(FileData);
+
+	if (Document.HasError())
+	{
+		// Have to load the file again as the JSON contents was rewritten by the insitu parse
+		ensure(FFileHelper::LoadFileToString(FileContents, *InJsonFile));
+		UE_LOG(LogInternationalizationArchiveSerializer, Error, TEXT("Failed to parse archive '%s'. %s."), *InJsonFile, *Document.GetError().CreateMessage(FileContents));
+		return false;
+	}
 
-	TSharedRef<TJsonReader<>> JsonReader = TJsonReaderFactory<>::Create(FileContents);
-	if (!FJsonSerializer::Deserialize(JsonReader, JsonObject) || !JsonObject.IsValid())
+	TOptional<UE::Json::FConstObject> RootObject = UE::Json::GetRootObject(Document.GetValue());
+	if (!RootObject.IsSet())
 	{
-		UE_LOG(LogInternationalizationArchiveSerializer, Error, TEXT("Failed to parse archive '%s'. %s."), *InJsonFile, *JsonReader->GetErrorMessage());
+		UE_LOG(LogInternationalizationArchiveSerializer, Error, TEXT("Failed to parse archive '%s'. No root object."), *InJsonFile);
 		return false;
 	}
 
-	return DeserializeInternal(JsonObject.ToSharedRef(), InArchive, InManifest, InNativeArchive);
+	return DeserializeInternal(*RootObject, InArchive, InManifest, InNativeArchive);
 }
 
 
-bool FJsonInternationalizationArchiveSerializer::SerializeArchive( TSharedRef< const FInternationalizationArchive > InArchive, FString& Str )
+bool FJsonInternationalizationArchiveSerializer::SerializeArchiveToString( TSharedRef< const FInternationalizationArchive > InArchive, FString& Str )
 {
-	TSharedRef< FJsonObject > JsonArchiveObj = MakeShareable( new FJsonObject );
-	bool bExecSuccessful = SerializeInternal( InArchive, JsonArchiveObj );
+	UE::Json::FDocument Document(rapidjson::kObjectType);
 
-	if( bExecSuccessful )
+	if (SerializeInternal(InArchive, Document.GetObject(), Document.GetAllocator()))
 	{
-		TSharedRef< TJsonWriter<> > Writer = TJsonWriterFactory<>::Create( &Str );
-		bExecSuccessful = FJsonSerializer::Serialize( JsonArchiveObj, Writer );
-		Writer->Close();
+		Str = UE::Json::WritePretty(Document);
+		return true;
 	}
 
-	return bExecSuccessful;
-}
-
-
-bool FJsonInternationalizationArchiveSerializer::SerializeArchive( TSharedRef< const FInternationalizationArchive > InArchive, TSharedRef< FJsonObject > InJsonObj )
-{
-	return SerializeInternal( InArchive, InJsonObj );
+	return false;
 }
 
 
 bool FJsonInternationalizationArchiveSerializer::SerializeArchiveToFile(TSharedRef<const FInternationalizationArchive> InArchive, const FString& InJsonFile)
 {
-	TSharedRef<FJsonObject> JsonObject = MakeShareable(new FJsonObject());
-	if (!SerializeArchive(InArchive, JsonObject))
+	FString JsonString;
+	if (!SerializeArchiveToString(InArchive, JsonString))
 	{
 		UE_LOG(LogInternationalizationArchiveSerializer, Error, TEXT("Failed to serialize archive '%s'."), *InJsonFile);
 		return false;
 	}
 
-	// Print the JSON data to a string
-	FString OutputJsonString;
-	TSharedRef<TJsonWriter<>> JsonWriter = TJsonWriterFactory<>::Create(&OutputJsonString);
-	FJsonSerializer::Serialize(JsonObject, JsonWriter);
-
 	// Save the JSON string (force Unicode for our manifest and archive files)
-	if (!FFileHelper::SaveStringToFile(OutputJsonString, *InJsonFile, FFileHelper::EEncodingOptions::ForceUnicode))
+	// TODO: Switch to UTF-8 by default, unless the file on-disk is already UTF-16
+	if (!FFileHelper::SaveStringToFile(JsonString, *InJsonFile, FFileHelper::EEncodingOptions::ForceUnicode))
 	{
 		UE_LOG(LogInternationalizationArchiveSerializer, Error, TEXT("Failed to save archive '%s'."), *InJsonFile);
 		return false;
@@ -119,18 +116,17 @@ bool FJsonInternationalizationArchiveSerializer::SerializeArchiveToFile(TSharedR
 }
 
 
-bool FJsonInternationalizationArchiveSerializer::DeserializeInternal(TSharedRef<FJsonObject> InJsonObj, TSharedRef<FInternationalizationArchive> InArchive, TSharedPtr<const FInternationalizationManifest> InManifest, TSharedPtr<const FInternationalizationArchive> InNativeArchive)
+bool FJsonInternationalizationArchiveSerializer::DeserializeInternal(UE::Json::FConstObject InJsonObj, TSharedRef<FInternationalizationArchive> InArchive, TSharedPtr<const FInternationalizationManifest> InManifest, TSharedPtr<const FInternationalizationArchive> InNativeArchive)
 {
-	if (InJsonObj->HasField(TAG_FORMATVERSION))
+	if (TOptional<int32> FormatVersion = UE::Json::GetInt32Field(InJsonObj, *TAG_FORMATVERSION))
 	{
-		const int32 FormatVersion = static_cast<int32>(InJsonObj->GetNumberField(TAG_FORMATVERSION));
-		if (FormatVersion > (int32)FInternationalizationArchive::EFormatVersion::Latest)
+		if (*FormatVersion > (int32)FInternationalizationArchive::EFormatVersion::Latest)
 		{
 			// Archive is too new to be loaded!
 			return false;
 		}
 
-		InArchive->SetFormatVersion(static_cast<FInternationalizationArchive::EFormatVersion>(FormatVersion));
+		InArchive->SetFormatVersion(static_cast<FInternationalizationArchive::EFormatVersion>(*FormatVersion));
 	}
 	else
 	{
@@ -154,75 +150,72 @@ bool FJsonInternationalizationArchiveSerializer::DeserializeInternal(TSharedRef<
 }
 
 
-bool FJsonInternationalizationArchiveSerializer::SerializeInternal(TSharedRef<const FInternationalizationArchive> InArchive, TSharedRef<FJsonObject> JsonObj)
+bool FJsonInternationalizationArchiveSerializer::SerializeInternal(TSharedRef<const FInternationalizationArchive> InArchive, UE::Json::FObject JsonObj, UE::Json::FAllocator& Allocator)
 {
-	TSharedPtr<FStructuredArchiveEntry> RootElement = MakeShareable(new FStructuredArchiveEntry(FString()));
+	TSharedPtr<FStructuredArchiveEntry> RootElement = MakeShared<FStructuredArchiveEntry>(FString());
 
 	// Condition the data so that it exists in a structured hierarchy for easy population of the JSON object.
 	GenerateStructuredData(InArchive, RootElement);
 
 	SortStructuredData(RootElement);
 
-	// Clear anything that may be in the JSON object
-	JsonObj->Values.Empty();
+	// Clear out anything that may be in the JSON object
+	JsonObj.RemoveAllMembers();
 
 	// Set format version.
-	JsonObj->SetNumberField(TAG_FORMATVERSION, static_cast<double>(InArchive->GetFormatVersion()));
+	JsonObj.AddMember(UE::Json::MakeStringRef(TAG_FORMATVERSION), UE::Json::FValue((int32)InArchive->GetFormatVersion()), Allocator);
 
 	// Setup the JSON object using the structured data created
-	StructuredDataToJsonObj(RootElement, JsonObj);
+	StructuredDataToJsonObj(RootElement, JsonObj, Allocator);
+
 	return true;
 }
 
 
-bool FJsonInternationalizationArchiveSerializer::JsonObjToArchive(TSharedRef<FJsonObject> InJsonObj, const FString& ParentNamespace, TSharedRef<FInternationalizationArchive> InArchive, TSharedPtr<const FInternationalizationManifest> InManifest, TSharedPtr<const FInternationalizationArchive> InNativeArchive)
+bool FJsonInternationalizationArchiveSerializer::JsonObjToArchive(UE::Json::FConstObject InJsonObj, const FString& ParentNamespace, TSharedRef<FInternationalizationArchive> InArchive, TSharedPtr<const FInternationalizationManifest> InManifest, TSharedPtr<const FInternationalizationArchive> InNativeArchive)
 {
-	bool bConvertSuccess = true;
 	FString AccumulatedNamespace = ParentNamespace;
 
-	if (InJsonObj->HasField(TAG_NAMESPACE))
+	if (TOptional<FStringView> Namespace = UE::Json::GetStringField(InJsonObj, *TAG_NAMESPACE))
 	{
-		if (!(AccumulatedNamespace.IsEmpty()))
+		if (!AccumulatedNamespace.IsEmpty())
 		{
 			AccumulatedNamespace += NAMESPACE_DELIMITER;
 		}
-		AccumulatedNamespace += InJsonObj->GetStringField(TAG_NAMESPACE);
+		AccumulatedNamespace += *Namespace;
 	}
 	else
 	{
 		UE_LOG(LogInternationalizationArchiveSerializer, Warning, TEXT("Encountered an object with a missing namespace while converting to Internationalization archive."));
-		bConvertSuccess = false;
+		return false;
 	}
 
 	// Process all the child objects
-	if (bConvertSuccess && InJsonObj->HasField(TAG_CHILDREN))
+	if (TOptional<UE::Json::FConstArray> ChildrenArray = UE::Json::GetArrayField(InJsonObj, *TAG_CHILDREN))
 	{
-		const TArray< TSharedPtr<FJsonValue> > ChildrenArray = InJsonObj->GetArrayField(TAG_CHILDREN);
-
-		for (TArray< TSharedPtr< FJsonValue > >::TConstIterator ChildIter(ChildrenArray.CreateConstIterator()); ChildIter; ++ChildIter)
+		for (const UE::Json::FValue& ChildEntry : *ChildrenArray)
 		{
-			const TSharedPtr< FJsonValue >  ChildEntry = *ChildIter;
-			const TSharedPtr< FJsonObject > ChildJSONObject = ChildEntry->AsObject();
+			if (!ChildEntry.IsObject())
+			{
+				return false;
+			}
+
+			UE::Json::FConstObject ChildJSONObject = ChildEntry.GetObject();
 
 			FString SourceText;
 			TSharedPtr< FLocMetadataObject > SourceMetadata;
-			if (ChildJSONObject->HasTypedField< EJson::String >(TAG_DEPRECATED_DEFAULTTEXT))
-			{
-				SourceText = ChildJSONObject->GetStringField(TAG_DEPRECATED_DEFAULTTEXT);
-			}
-			else if (ChildJSONObject->HasTypedField< EJson::Object >(TAG_SOURCE))
+			if (TOptional<UE::Json::FConstObject> SourceJSONObject = UE::Json::GetObjectField(ChildJSONObject, *TAG_SOURCE))
 			{
-				const TSharedPtr< FJsonObject > SourceJSONObject = ChildJSONObject->GetObjectField(TAG_SOURCE);
-				if (SourceJSONObject->HasTypedField< EJson::String >(TAG_SOURCE_TEXT))
+				if (TOptional<FStringView> SourceTextValue = UE::Json::GetStringField(*SourceJSONObject, *TAG_SOURCE_TEXT))
 				{
-					SourceText = SourceJSONObject->GetStringField(TAG_SOURCE_TEXT);
+					SourceText = *SourceTextValue;
 
 					// Source meta data is mixed in with the source text, we'll process metadata if the source json object has more than one entry
-					if (SourceJSONObject->Values.Num() > 1)
+					if (SourceJSONObject->MemberCount() > 1)
 					{
 						// We load in the entire source object as metadata and just remove the source text.
-						FJsonInternationalizationMetaDataSerializer::DeserializeMetadata(SourceJSONObject.ToSharedRef(), SourceMetadata);
-						if (SourceMetadata.IsValid())
+						FJsonInternationalizationMetaDataSerializer::DeserializeMetadata(*SourceJSONObject, SourceMetadata);
+						if (SourceMetadata)
 						{
 							SourceMetadata->Values.Remove(TAG_SOURCE_TEXT);
 						}
@@ -230,128 +223,124 @@ bool FJsonInternationalizationArchiveSerializer::JsonObjToArchive(TSharedRef<FJs
 				}
 				else
 				{
-					bConvertSuccess = false;
+					return false;
 				}
 			}
+			else if (TOptional<FStringView> DefaultText = UE::Json::GetStringField(ChildJSONObject, *TAG_DEPRECATED_DEFAULTTEXT))
+			{
+				SourceText = *DefaultText;
+			}
 			else
 			{
-				bConvertSuccess = false;
+				return false;
 			}
 
 			FString TranslationText;
 			TSharedPtr< FLocMetadataObject > TranslationMetadata;
-			if (ChildJSONObject->HasTypedField< EJson::String >(TAG_DEPRECATED_TRANSLATEDTEXT))
+			if (TOptional<UE::Json::FConstObject> TranslationJSONObject = UE::Json::GetObjectField(ChildJSONObject, *TAG_TRANSLATION))
 			{
-				TranslationText = ChildJSONObject->GetStringField(TAG_DEPRECATED_TRANSLATEDTEXT);
-			}
-			else if (ChildJSONObject->HasTypedField< EJson::Object >(TAG_TRANSLATION))
-			{
-				const TSharedPtr< FJsonObject > TranslationJSONObject = ChildJSONObject->GetObjectField(TAG_TRANSLATION);
-				if (TranslationJSONObject->HasTypedField< EJson::String >(TAG_TRANSLATION_TEXT))
+				if (TOptional<FStringView> TranslationTextValue = UE::Json::GetStringField(*TranslationJSONObject, *TAG_TRANSLATION_TEXT))
 				{
-					TranslationText = TranslationJSONObject->GetStringField(TAG_TRANSLATION_TEXT);
+					TranslationText = *TranslationTextValue;
 
-					// Source meta data is mixed in with the source text, we'll process metadata if the source json object has more than one entry
-					if (TranslationJSONObject->Values.Num() > 1)
+					// Translation meta data is mixed in with the translation text, we'll process metadata if the source json object has more than one entry
+					if (TranslationJSONObject->MemberCount() > 1)
 					{
-						// We load in the entire source object as metadata and remove the source text
-						FJsonInternationalizationMetaDataSerializer::DeserializeMetadata(TranslationJSONObject.ToSharedRef(), TranslationMetadata);
-						if (TranslationJSONObject.IsValid())
+						// We load in the entire translation object as metadata and remove the translation text
+						FJsonInternationalizationMetaDataSerializer::DeserializeMetadata(*TranslationJSONObject, TranslationMetadata);
+						if (TranslationMetadata)
 						{
-							TranslationJSONObject->Values.Remove(TAG_TRANSLATION_TEXT);
+							TranslationMetadata->Values.Remove(TAG_TRANSLATION_TEXT);
 						}
 					}
 				}
+				else if (TOptional<FStringView> TranslatedText = UE::Json::GetStringField(ChildJSONObject, *TAG_DEPRECATED_TRANSLATEDTEXT))
+				{
+					TranslationText = *TranslatedText;
+				}
 				else
 				{
-					bConvertSuccess = false;
+					return false;
 				}
 			}
 			else
 			{
-				bConvertSuccess = false;
+				return false;
 			}
 
-			if (bConvertSuccess)
-			{
-				FLocItem Source(SourceText);
-				Source.MetadataObj = SourceMetadata;
+			FLocItem Source(SourceText);
+			Source.MetadataObj = SourceMetadata;
 
-				FLocItem Translation(TranslationText);
-				Translation.MetadataObj = TranslationMetadata;
+			FLocItem Translation(TranslationText);
+			Translation.MetadataObj = TranslationMetadata;
 
-				bool bIsOptional = false;
-				if (ChildJSONObject->HasTypedField< EJson::Boolean >(TAG_OPTIONAL))
+			bool bIsOptional = false;
+			if (TOptional<bool> IsOptionalValue = UE::Json::GetBoolField(ChildJSONObject, *TAG_OPTIONAL))
+			{
+				bIsOptional = *IsOptionalValue;
+			}
+
+			TArray<FLocKey> Keys;
+			TSharedPtr< FLocMetadataObject > KeyMetadataNode;
+			if (InArchive->GetFormatVersion() < FInternationalizationArchive::EFormatVersion::AddedKeys)
+			{
+				// We used to store the key meta-data as a top-level value, rather than within a "MetaData" object
+				if (TOptional<UE::Json::FConstObject> MetaDataKeyJSONObject = UE::Json::GetObjectField(ChildJSONObject, *TAG_METADATA_KEY))
 				{
-					bIsOptional = ChildJSONObject->GetBoolField(TAG_OPTIONAL);
+					FJsonInternationalizationMetaDataSerializer::DeserializeMetadata(*MetaDataKeyJSONObject, KeyMetadataNode);
 				}
 
-				TArray<FLocKey> Keys;
-				TSharedPtr< FLocMetadataObject > KeyMetadataNode;
-				if (InArchive->GetFormatVersion() < FInternationalizationArchive::EFormatVersion::AddedKeys)
+				if (InManifest)
 				{
-					// We used to store the key meta-data as a top-level value, rather than within a "MetaData" object
-					if (ChildJSONObject->HasTypedField< EJson::Object >(TAG_METADATA_KEY))
-					{
-						const TSharedPtr< FJsonObject > MetaDataKeyJSONObject = ChildJSONObject->GetObjectField(TAG_METADATA_KEY);
-						FJsonInternationalizationMetaDataSerializer::DeserializeMetadata(MetaDataKeyJSONObject.ToSharedRef(), KeyMetadataNode);
-					}
-
-					if (InManifest.IsValid())
-					{
-						// We have no key in the archive data, so we must try and infer it from the manifest
-						FLocTextHelper::FindKeysForLegacyTranslation(InManifest.ToSharedRef(), InNativeArchive, AccumulatedNamespace, SourceText, KeyMetadataNode, Keys);
-					}
+					// We have no key in the archive data, so we must try and infer it from the manifest
+					FLocTextHelper::FindKeysForLegacyTranslation(InManifest.ToSharedRef(), InNativeArchive, AccumulatedNamespace, SourceText, KeyMetadataNode, Keys);
 				}
-				else
+			}
+			else
+			{
+				if (TOptional<FStringView> KeyValue = UE::Json::GetStringField(ChildJSONObject, *TAG_KEY))
 				{
-					if (ChildJSONObject->HasTypedField< EJson::String >(TAG_KEY))
-					{
-						Keys.Add(ChildJSONObject->GetStringField(TAG_KEY));
-					}
+					Keys.Add(FString(*KeyValue));
+				}
 
-					if (ChildJSONObject->HasTypedField< EJson::Object >(TAG_METADATA))
+				if (TOptional<UE::Json::FConstObject> MetaDataJSONObject = UE::Json::GetObjectField(ChildJSONObject, *TAG_METADATA))
+				{
+					if (TOptional<UE::Json::FConstObject> MetaDataKeyJSONObject = UE::Json::GetObjectField(*MetaDataJSONObject, *TAG_METADATA_KEY))
 					{
-						const TSharedPtr< FJsonObject > MetaDataJSONObject = ChildJSONObject->GetObjectField(TAG_METADATA);
-
-						if (MetaDataJSONObject->HasTypedField< EJson::Object >(TAG_METADATA_KEY))
-						{
-							const TSharedPtr< FJsonObject > MetaDataKeyJSONObject = MetaDataJSONObject->GetObjectField(TAG_METADATA_KEY);
-							FJsonInternationalizationMetaDataSerializer::DeserializeMetadata(MetaDataKeyJSONObject.ToSharedRef(), KeyMetadataNode);
-						}
+						FJsonInternationalizationMetaDataSerializer::DeserializeMetadata(*MetaDataKeyJSONObject, KeyMetadataNode);
 					}
 				}
+			}
 
-				for (const FLocKey& Key : Keys)
+			for (const FLocKey& Key : Keys)
+			{
+				const bool bAddSuccessful = InArchive->AddEntry(AccumulatedNamespace, Key, Source, Translation, KeyMetadataNode, bIsOptional);
+				if (!bAddSuccessful)
 				{
-					const bool bAddSuccessful = InArchive->AddEntry(AccumulatedNamespace, Key, Source, Translation, KeyMetadataNode, bIsOptional);
-					if (!bAddSuccessful)
-					{
-						UE_LOG(LogInternationalizationArchiveSerializer, Warning, TEXT("Could not add JSON entry to the Internationalization archive: Namespace:%s Key:%s DefaultText:%s"), *AccumulatedNamespace, *Key.GetString(), *SourceText);
-					}
+					UE_LOG(LogInternationalizationArchiveSerializer, Warning, TEXT("Could not add JSON entry to the Internationalization archive: Namespace:%s Key:%s DefaultText:%s"), *AccumulatedNamespace, *Key.GetString(), *SourceText);
 				}
 			}
 		}
 	}
 
-	if (bConvertSuccess && InJsonObj->HasField(TAG_SUBNAMESPACES))
+	if (TOptional<UE::Json::FConstArray> SubnamespaceArray = UE::Json::GetArrayField(InJsonObj, *TAG_SUBNAMESPACES))
 	{
-		const TArray< TSharedPtr<FJsonValue> > SubnamespaceArray = InJsonObj->GetArrayField(TAG_SUBNAMESPACES);
-
-		for (TArray< TSharedPtr< FJsonValue > >::TConstIterator SubnamespaceIter(SubnamespaceArray.CreateConstIterator()); SubnamespaceIter; ++SubnamespaceIter)
+		for (const UE::Json::FValue& SubnamespaceEntry : *SubnamespaceArray)
 		{
-			const TSharedPtr< FJsonValue >  SubnamespaceEntry = *SubnamespaceIter;
-			const TSharedPtr< FJsonObject > SubnamespaceJSONObject = SubnamespaceEntry->AsObject();
+			if (!SubnamespaceEntry.IsObject())
+			{
+				return false;
+			}
 
-			if (!JsonObjToArchive(SubnamespaceJSONObject.ToSharedRef(), AccumulatedNamespace, InArchive, InManifest, InNativeArchive))
+			UE::Json::FConstObject SubnamespaceJSONObject = SubnamespaceEntry.GetObject();
+			if (!JsonObjToArchive(SubnamespaceJSONObject, AccumulatedNamespace, InArchive, InManifest, InNativeArchive))
 			{
-				bConvertSuccess = false;
-				break;
+				return false;
 			}
 		}
 	}
 
-	return bConvertSuccess;
+	return true;
 }
 
 
@@ -433,94 +422,92 @@ void FJsonInternationalizationArchiveSerializer::SortStructuredData( TSharedPtr<
 }
 
 
-void FJsonInternationalizationArchiveSerializer::StructuredDataToJsonObj( TSharedPtr< const FStructuredArchiveEntry > InElement, TSharedRef< FJsonObject > OutJsonObj )
+void FJsonInternationalizationArchiveSerializer::StructuredDataToJsonObj( TSharedPtr< const FStructuredArchiveEntry > InElement, UE::Json::FObject JsonObj, UE::Json::FAllocator& Allocator )
 {
-	OutJsonObj->SetStringField( TAG_NAMESPACE, InElement->Namespace );
-
-	TArray< TSharedPtr< FJsonValue > > NamespaceArray;
-	TArray< TSharedPtr< FJsonValue > > EntryArray;
+	JsonObj.AddMember(UE::Json::MakeStringRef(TAG_NAMESPACE), UE::Json::MakeStringValue(InElement->Namespace, Allocator), Allocator);
 
-	//Write namespace content entries
-	for( TArray< TSharedPtr< FArchiveEntry > >::TConstIterator Iter( InElement->ArchiveEntries ); Iter; ++Iter )
+	// Write namespace content entries
+	UE::Json::FValue EntryJSONArray(rapidjson::kArrayType);
+	for (const TSharedPtr<FArchiveEntry>& Entry : InElement->ArchiveEntries)
 	{
-		const TSharedPtr< FArchiveEntry > Entry = *Iter;
-		TSharedPtr< FJsonObject > EntryNode = MakeShareable( new FJsonObject );
-
-		FString ProcessedSourceText = Entry->Source.Text;
-		FString ProcessedTranslation = Entry->Translation.Text;
+		UE::Json::FValue EntryJSONObject(rapidjson::kObjectType);
 
-		TSharedPtr< FJsonObject > SourceNode;
-		if( Entry->Source.MetadataObj.IsValid() )
 		{
-			FJsonInternationalizationMetaDataSerializer::SerializeMetadata( Entry->Source.MetadataObj.ToSharedRef(), SourceNode );
-		}
+			UE::Json::FValue SourceJSONObject(rapidjson::kObjectType);
 
-		if( !SourceNode.IsValid() )
-		{
-			SourceNode = MakeShareable( new FJsonObject );
-		}
+			if (Entry->Source.MetadataObj)
+			{
+				FJsonInternationalizationMetaDataSerializer::SerializeMetadata(Entry->Source.MetadataObj.ToSharedRef(), SourceJSONObject, Allocator);
+			}
+
+			SourceJSONObject.AddMember(UE::Json::MakeStringRef(TAG_SOURCE_TEXT), UE::Json::MakeStringValue(Entry->Source.Text, Allocator), Allocator);
 
-		SourceNode->SetStringField( TAG_SOURCE_TEXT, ProcessedSourceText );
-		EntryNode->SetObjectField( TAG_SOURCE, SourceNode );
+			EntryJSONObject.AddMember(UE::Json::MakeStringRef(TAG_SOURCE), MoveTemp(SourceJSONObject), Allocator);
+		}
 		
-		TSharedPtr< FJsonObject > TranslationNode;
-		if( Entry->Translation.MetadataObj.IsValid() )
 		{
-			FJsonInternationalizationMetaDataSerializer::SerializeMetadata( Entry->Translation.MetadataObj.ToSharedRef(), TranslationNode );
-		}
+			UE::Json::FValue TranslationJSONObject(rapidjson::kObjectType);
 
-		if( !TranslationNode.IsValid() )
-		{
-			TranslationNode = MakeShareable( new FJsonObject );
-		}
+			if (Entry->Translation.MetadataObj)
+			{
+				FJsonInternationalizationMetaDataSerializer::SerializeMetadata(Entry->Translation.MetadataObj.ToSharedRef(), TranslationJSONObject, Allocator);
+			}
+
+			TranslationJSONObject.AddMember(UE::Json::MakeStringRef(TAG_TRANSLATION_TEXT), UE::Json::MakeStringValue(Entry->Translation.Text, Allocator), Allocator);
 
-		TranslationNode->SetStringField( TAG_TRANSLATION_TEXT, ProcessedTranslation );
-		EntryNode->SetObjectField( TAG_TRANSLATION, TranslationNode );
+			EntryJSONObject.AddMember(UE::Json::MakeStringRef(TAG_TRANSLATION), MoveTemp(TranslationJSONObject), Allocator);
+		}
 
-		EntryNode->SetStringField( TAG_KEY, Entry->Key.GetString() );
+		EntryJSONObject.AddMember(UE::Json::MakeStringRef(TAG_KEY), UE::Json::MakeStringValue(Entry->Key.GetString(), Allocator), Allocator);
 
-		if( Entry->KeyMetadataObj.IsValid() )
+		if (Entry->KeyMetadataObj)
 		{
-			TSharedRef< FJsonObject > MetaDataNode = MakeShareable(new FJsonObject());
-			EntryNode->SetObjectField( TAG_METADATA, MetaDataNode );
+			UE::Json::FValue MetaDataJSONObject(rapidjson::kObjectType);
+
+			{
+				UE::Json::FValue KeyDataJSONObject(rapidjson::kObjectType);
+				FJsonInternationalizationMetaDataSerializer::SerializeMetadata(Entry->KeyMetadataObj.ToSharedRef(), KeyDataJSONObject, Allocator);
+				if (KeyDataJSONObject.MemberCount() > 0)
+				{
+					MetaDataJSONObject.AddMember(UE::Json::MakeStringRef(TAG_METADATA_KEY), MoveTemp(KeyDataJSONObject), Allocator);
+				}
+			}
 
-			TSharedPtr< FJsonObject > KeyMetaDataNode;
-			FJsonInternationalizationMetaDataSerializer::SerializeMetadata( Entry->KeyMetadataObj.ToSharedRef(), KeyMetaDataNode );
-			if( KeyMetaDataNode.IsValid() )
+			if (MetaDataJSONObject.MemberCount() > 0)
 			{
-				MetaDataNode->SetObjectField( TAG_METADATA_KEY, KeyMetaDataNode );
+				EntryJSONObject.AddMember(UE::Json::MakeStringRef(TAG_METADATA), MoveTemp(MetaDataJSONObject), Allocator);
 			}
 		}
 
 		// We only add the optional field if it is true, it is assumed to be false otherwise.
-		if( Entry->bIsOptional == true )
+		if (Entry->bIsOptional)
 		{
-			EntryNode->SetBoolField( TAG_OPTIONAL, Entry->bIsOptional );
+			EntryJSONObject.AddMember(UE::Json::MakeStringRef(TAG_OPTIONAL), UE::Json::FValue(Entry->bIsOptional), Allocator);
 		}
 
-		EntryArray.Add( MakeShareable( new FJsonValueObject( EntryNode ) ) );
+		EntryJSONArray.PushBack(MoveTemp(EntryJSONObject), Allocator);
 	}
 
-	//Write the subnamespaces
-	for( auto Iter = InElement->SubNamespaces.CreateConstIterator(); Iter; ++Iter )
+	// Write the subnamespaces
+	UE::Json::FValue NamespaceJSONArray(rapidjson::kArrayType);
+	for (const TTuple<FString, TSharedPtr<FStructuredArchiveEntry>>& SubElementPair : InElement->SubNamespaces)
 	{
-		const TSharedPtr<FStructuredArchiveEntry> SubElement = Iter->Value;
-		if( SubElement.IsValid() )
+		if (SubElementPair.Value)
 		{
-			TSharedRef<FJsonObject> SubObject = MakeShareable( new FJsonObject );
-			StructuredDataToJsonObj( SubElement, SubObject );
+			UE::Json::FValue SubJSONObject(rapidjson::kObjectType);
+			StructuredDataToJsonObj(SubElementPair.Value, SubJSONObject.GetObject(), Allocator);
 
-			NamespaceArray.Add( MakeShareable( new FJsonValueObject( SubObject ) ) );
+			NamespaceJSONArray.PushBack(MoveTemp(SubJSONObject), Allocator);
 		}
 	}
 
-	if( EntryArray.Num() > 0 )
+	if (EntryJSONArray.Size() > 0)
 	{
-		OutJsonObj->SetArrayField( TAG_CHILDREN, EntryArray );
+		JsonObj.AddMember(UE::Json::MakeStringRef(TAG_CHILDREN), MoveTemp(EntryJSONArray), Allocator);
 	}
 
-	if( NamespaceArray.Num() > 0 )
+	if (NamespaceJSONArray.Size() > 0)
 	{
-		OutJsonObj->SetArrayField( TAG_SUBNAMESPACES, NamespaceArray );
+		JsonObj.AddMember(UE::Json::MakeStringRef(TAG_SUBNAMESPACES), MoveTemp(NamespaceJSONArray), Allocator);
 	}
 }
diff --git a/Engine/Source/Developer/Localization/Private/Serialization/JsonInternationalizationManifestSerializer.cpp b/Engine/Source/Developer/Localization/Private/Serialization/JsonInternationalizationManifestSerializer.cpp
--- a/Engine/Source/Developer/Localization/Private/Serialization/JsonInternationalizationManifestSerializer.cpp
+++ b/Engine/Source/Developer/Localization/Private/Serialization/JsonInternationalizationManifestSerializer.cpp
@@ -3,8 +3,6 @@
 #include "Serialization/JsonInternationalizationManifestSerializer.h"
 #include "Misc/FileHelper.h"
 #include "Misc/Paths.h"
-#include "Serialization/JsonReader.h"
-#include "Serialization/JsonSerializer.h"
 #include "Serialization/JsonInternationalizationMetadataSerializer.h"
 
 
@@ -28,24 +26,24 @@ const FString FJsonInternationalizationManifestSerializer::TAG_METADATA_KEY = TE
 const FString FJsonInternationalizationManifestSerializer::NAMESPACE_DELIMITER = TEXT(".");
 
 
-bool FJsonInternationalizationManifestSerializer::DeserializeManifest( const FString& InStr, TSharedRef< FInternationalizationManifest > Manifest, const FName PlatformName )
+bool FJsonInternationalizationManifestSerializer::DeserializeManifestFromString( const FString& InStr, TSharedRef< FInternationalizationManifest > Manifest, const FName PlatformName )
 {
-	TSharedPtr<FJsonObject> JsonManifestObj;
-	TSharedRef< TJsonReader<> > Reader = TJsonReaderFactory<>::Create( InStr );
-	bool bExecSuccessful = FJsonSerializer::Deserialize( Reader, JsonManifestObj );
+	TValueOrError<UE::Json::FDocument, UE::Json::FParseError> Document = UE::Json::Parse(InStr);
 
-	if ( bExecSuccessful && JsonManifestObj.IsValid() )
+	if (Document.HasError())
 	{
-		bExecSuccessful = DeserializeInternal( JsonManifestObj.ToSharedRef(), Manifest, PlatformName );
+		UE_LOG(LogInternationalizationManifestSerializer, Error, TEXT("Failed to parse manifest. %s."), *Document.GetError().CreateMessage(InStr));
+		return false;
 	}
 
-	return bExecSuccessful;
-}
-
+	TOptional<UE::Json::FConstObject> RootObject = UE::Json::GetRootObject(Document.GetValue());
+	if (!RootObject.IsSet())
+	{
+		UE_LOG(LogInternationalizationManifestSerializer, Error, TEXT("Failed to parse manifest. No root object."));
+		return false;
+	}
 
-bool FJsonInternationalizationManifestSerializer::DeserializeManifest( TSharedRef< FJsonObject > InJsonObj, TSharedRef< FInternationalizationManifest > Manifest, const FName PlatformName )
-{
-	return DeserializeInternal( InJsonObj, Manifest, PlatformName );
+	return DeserializeInternal(*RootObject, Manifest, PlatformName);
 }
 
 
@@ -59,57 +57,55 @@ bool FJsonInternationalizationManifestSerializer::DeserializeManifestFromFile( c
 		return false;
 	}
 
-	// Parse as JSON
-	TSharedPtr<FJsonObject> JsonObject;
+	// Grab the internal character array so we can do an insitu parse
+	TArray<TCHAR> FileData = MoveTemp(FileContents.GetCharArray());
+	TValueOrError<UE::Json::FDocument, UE::Json::FParseError> Document = UE::Json::ParseInPlace(FileData);
+
+	if (Document.HasError())
+	{
+		// Have to load the file again as the JSON contents was rewritten by the insitu parse
+		ensure(FFileHelper::LoadFileToString(FileContents, *InJsonFile));
+		UE_LOG(LogInternationalizationManifestSerializer, Error, TEXT("Failed to parse manifest '%s'. %s."), *InJsonFile, *Document.GetError().CreateMessage(FileContents));
+		return false;
+	}
 
-	TSharedRef<TJsonReader<>> JsonReader = TJsonReaderFactory<>::Create(FileContents);
-	if (!FJsonSerializer::Deserialize(JsonReader, JsonObject) || !JsonObject.IsValid())
+	TOptional<UE::Json::FConstObject> RootObject = UE::Json::GetRootObject(Document.GetValue());
+	if (!RootObject.IsSet())
 	{
-		UE_LOG(LogInternationalizationManifestSerializer, Error, TEXT("Failed to parse manifest '%s'. %s."), *InJsonFile, *JsonReader->GetErrorMessage());
+		UE_LOG(LogInternationalizationManifestSerializer, Error, TEXT("Failed to parse manifest '%s'. No root object."), *InJsonFile);
 		return false;
 	}
 
-	return DeserializeInternal(JsonObject.ToSharedRef(), Manifest, PlatformName);
+	return DeserializeInternal(*RootObject, Manifest, PlatformName);
 }
 
 
-bool FJsonInternationalizationManifestSerializer::SerializeManifest( TSharedRef< const FInternationalizationManifest > Manifest, FString& Str )
+bool FJsonInternationalizationManifestSerializer::SerializeManifestToString( TSharedRef< const FInternationalizationManifest > Manifest, FString& Str )
 {
-	TSharedRef< FJsonObject > JsonManifestObj = MakeShareable( new FJsonObject );
-	bool bExecSuccessful = SerializeInternal( Manifest, JsonManifestObj );
-
-	if(bExecSuccessful)
+	UE::Json::FDocument Document(rapidjson::kObjectType);
+	
+	if (SerializeInternal(Manifest, Document.GetObject(), Document.GetAllocator()))
 	{
-		TSharedRef< TJsonWriter<> > Writer = TJsonWriterFactory<>::Create( &Str );
-		bExecSuccessful = FJsonSerializer::Serialize(JsonManifestObj, Writer);
-		Writer->Close();
+		Str = UE::Json::WritePretty(Document);
+		return true;
 	}
-	return bExecSuccessful;
-}
 
-
-bool FJsonInternationalizationManifestSerializer::SerializeManifest( TSharedRef< const FInternationalizationManifest > Manifest, TSharedRef< FJsonObject > JsonObj )
-{
-	return SerializeInternal( Manifest, JsonObj );
+	return false;
 }
 
 
 bool FJsonInternationalizationManifestSerializer::SerializeManifestToFile( TSharedRef< const FInternationalizationManifest > Manifest, const FString& InJsonFile )
 {
-	TSharedRef<FJsonObject> JsonObject = MakeShareable(new FJsonObject());
-	if (!SerializeManifest(Manifest, JsonObject))
+	FString JsonString;
+	if (!SerializeManifestToString(Manifest, JsonString))
 	{
 		UE_LOG(LogInternationalizationManifestSerializer, Error, TEXT("Failed to serialize manifest '%s'."), *InJsonFile);
 		return false;
 	}
 
-	// Print the JSON data to a string
-	FString OutputJsonString;
-	TSharedRef<TJsonWriter<>> JsonWriter = TJsonWriterFactory<>::Create(&OutputJsonString);
-	FJsonSerializer::Serialize(JsonObject, JsonWriter);
-
 	// Save the JSON string (force Unicode for our manifest and archive files)
-	if (!FFileHelper::SaveStringToFile(OutputJsonString, *InJsonFile, FFileHelper::EEncodingOptions::ForceUnicode))
+	// TODO: Switch to UTF-8 by default, unless the file on-disk is already UTF-16
+	if (!FFileHelper::SaveStringToFile(JsonString, *InJsonFile, FFileHelper::EEncodingOptions::ForceUnicode))
 	{
 		UE_LOG(LogInternationalizationManifestSerializer, Error, TEXT("Failed to save manifest '%s'."), *InJsonFile);
 		return false;
@@ -119,31 +115,37 @@ bool FJsonInternationalizationManifestSerializer::SerializeManifestToFile( TShar
 }
 
 
-bool FJsonInternationalizationManifestSerializer::DeserializeInternal( TSharedRef< FJsonObject > InJsonObj, TSharedRef< FInternationalizationManifest > Manifest, const FName PlatformName )
+bool FJsonInternationalizationManifestSerializer::DeserializeInternal( UE::Json::FConstObject InJsonObj, TSharedRef< FInternationalizationManifest > Manifest, const FName PlatformName )
 {
-	if( InJsonObj->HasField( TAG_FORMATVERSION ) )
+	if (TOptional<int32> FormatVersion = UE::Json::GetInt32Field(InJsonObj, *TAG_FORMATVERSION))
 	{
-		const int32 FormatVersion = static_cast<int32>(InJsonObj->GetNumberField( TAG_FORMATVERSION ));
-		if (FormatVersion > (int32)FInternationalizationManifest::EFormatVersion::Latest)
+		if (*FormatVersion > (int32)FInternationalizationManifest::EFormatVersion::Latest)
 		{
 			// Manifest is too new to be loaded!
 			return false;
 		}
 
-		Manifest->SetFormatVersion(static_cast<FInternationalizationManifest::EFormatVersion>(FormatVersion));
+		Manifest->SetFormatVersion(static_cast<FInternationalizationManifest::EFormatVersion>(*FormatVersion));
 	}
 	else
 	{
 		Manifest->SetFormatVersion(FInternationalizationManifest::EFormatVersion::Initial);
 	}
 
-	return JsonObjToManifest(InJsonObj, TEXT(""), Manifest, PlatformName );
+	if (JsonObjToManifest(InJsonObj, FString(), Manifest, PlatformName))
+	{
+		// We've been upgraded to the latest format now...
+		Manifest->SetFormatVersion(FInternationalizationManifest::EFormatVersion::Latest);
+		return true;
+	}
+
+	return false;
 }
 
 
-bool FJsonInternationalizationManifestSerializer::SerializeInternal( TSharedRef< const FInternationalizationManifest > InManifest, TSharedRef< FJsonObject > JsonObj )
+bool FJsonInternationalizationManifestSerializer::SerializeInternal( TSharedRef< const FInternationalizationManifest > InManifest, UE::Json::FObject JsonObj, UE::Json::FAllocator& Allocator )
 {
-	TSharedPtr< FStructuredEntry > RootElement = MakeShareable( new FStructuredEntry( TEXT("") ) );
+	TSharedPtr< FStructuredEntry > RootElement = MakeShared<FStructuredEntry>(FString());
 
 	// Condition the data so that it exists in a structured hierarchy for easy population of the JSON object.
 	GenerateStructuredData( InManifest, RootElement );
@@ -151,152 +153,146 @@ bool FJsonInternationalizationManifestSerializer::SerializeInternal( TSharedRef<
 	// Arrange the entries in non-cultural format so that diffs are easier to read.
 	SortStructuredData( RootElement );
 
-	//Clear out anything that may be in the JSON object
-	JsonObj->Values.Empty();
+	// Clear out anything that may be in the JSON object
+	JsonObj.RemoveAllMembers();
 
 	// Set format version.
-	JsonObj->SetNumberField(TAG_FORMATVERSION, static_cast<double>(InManifest->GetFormatVersion()));
+	JsonObj.AddMember(UE::Json::MakeStringRef(TAG_FORMATVERSION), UE::Json::FValue((int32)InManifest->GetFormatVersion()), Allocator);
 
 	// Setup the JSON object using the structured data created
-	StructuredDataToJsonObj( RootElement, JsonObj );
+	StructuredDataToJsonObj( RootElement, JsonObj, Allocator );
 
 	return true;
 }
 
 
-bool FJsonInternationalizationManifestSerializer::JsonObjToManifest( TSharedRef< FJsonObject > InJsonObj, FString ParentNamespace, TSharedRef< FInternationalizationManifest > Manifest, const FName PlatformName )
+bool FJsonInternationalizationManifestSerializer::JsonObjToManifest( UE::Json::FConstObject InJsonObj, FString ParentNamespace, TSharedRef< FInternationalizationManifest > Manifest, const FName PlatformName )
 {
-	bool bConvertSuccess = true;
 	FString AccumulatedNamespace = ParentNamespace;
-	if( InJsonObj->HasField( TAG_NAMESPACE) )
+
+	if (TOptional<FStringView> Namespace = UE::Json::GetStringField(InJsonObj, *TAG_NAMESPACE))
 	{
-		if( !( AccumulatedNamespace.IsEmpty() ) )
+		if (!AccumulatedNamespace.IsEmpty())
 		{
 			AccumulatedNamespace += NAMESPACE_DELIMITER;
 		}
-		AccumulatedNamespace += InJsonObj->GetStringField( TAG_NAMESPACE );
+		AccumulatedNamespace += *Namespace;
 	}
 	else
 	{
 		// We found an entry with a missing namespace
-		bConvertSuccess = false;
+		return false;
 	}
 
 	// Process all the child objects
-	if( bConvertSuccess && InJsonObj->HasField( TAG_CHILDREN ) )
+	if (TOptional<UE::Json::FConstArray> ChildrenArray = UE::Json::GetArrayField(InJsonObj, *TAG_CHILDREN))
 	{
-		const TArray< TSharedPtr< FJsonValue> > ChildrenArray = InJsonObj->GetArrayField( TAG_CHILDREN );
-
-		for(TArray< TSharedPtr< FJsonValue > >::TConstIterator ChildIter( ChildrenArray.CreateConstIterator() ); ChildIter && bConvertSuccess; ++ChildIter)
+		for (const UE::Json::FValue& ChildEntry : *ChildrenArray)
 		{
-			const TSharedPtr< FJsonValue >  ChildEntry = *ChildIter;
-			const TSharedPtr< FJsonObject > ChildJSONObject = ChildEntry->AsObject();
+			if (!ChildEntry.IsObject())
+			{
+				return false;
+			}
+
+			UE::Json::FConstObject ChildJSONObject = ChildEntry.GetObject();
 
-			
 			FString SourceText;
 			TSharedPtr< FLocMetadataObject > SourceMetadata;
-			if( ChildJSONObject->HasTypedField< EJson::String>( TAG_DEPRECATED_DEFAULTTEXT) )
+			if (TOptional<UE::Json::FConstObject> SourceJSONObject = UE::Json::GetObjectField(ChildJSONObject, *TAG_SOURCE))
 			{
-				SourceText = ChildJSONObject->GetStringField( TAG_DEPRECATED_DEFAULTTEXT );
-			}
-			else if( ChildJSONObject->HasTypedField< EJson::Object>( TAG_SOURCE ) )
-			{
-				const TSharedPtr< FJsonObject > SourceJSONObject = ChildJSONObject->GetObjectField( TAG_SOURCE );
-				if( SourceJSONObject->HasTypedField< EJson::String >( TAG_SOURCE_TEXT ) )
+				if (TOptional<FStringView> SourceTextValue = UE::Json::GetStringField(*SourceJSONObject, *TAG_SOURCE_TEXT))
 				{
-					SourceText = SourceJSONObject->GetStringField( TAG_SOURCE_TEXT );
+					SourceText = *SourceTextValue;
 
 					// Source meta data is mixed in with the source text, we'll process metadata if the source json object has more than one entry
-					if( SourceJSONObject->Values.Num() > 1 )
+					if (SourceJSONObject->MemberCount() > 1)
 					{
-						// We load in the entire source object as metadata and remove the source object
-						FJsonInternationalizationMetaDataSerializer::DeserializeMetadata( SourceJSONObject.ToSharedRef(), SourceMetadata );
-						if( SourceMetadata.IsValid() )
+						// We load in the entire source object as metadata and just remove the source text.
+						FJsonInternationalizationMetaDataSerializer::DeserializeMetadata(*SourceJSONObject, SourceMetadata);
+						if (SourceMetadata)
 						{
-							SourceMetadata->Values.Remove( TAG_SOURCE_TEXT );
+							SourceMetadata->Values.Remove(TAG_SOURCE_TEXT);
 						}
 					}
 				}
 				else
 				{
-					bConvertSuccess = false;
+					return false;
 				}
 			}
+			else if (TOptional<FStringView> DefaultText = UE::Json::GetStringField(ChildJSONObject, *TAG_DEPRECATED_DEFAULTTEXT))
+			{
+				SourceText = *DefaultText;
+			}
 			else
 			{
-				bConvertSuccess = false;
+				return false;
 			}
 
 			FLocItem Source(SourceText);
 			Source.MetadataObj = SourceMetadata;
 
-			if( bConvertSuccess && ChildJSONObject->HasField( TAG_KEYCOLLECTION ) )
+			if (TOptional<UE::Json::FConstArray> ContextArray = UE::Json::GetArrayField(ChildJSONObject, *TAG_KEYCOLLECTION))
 			{
-				
-				const TArray< TSharedPtr<FJsonValue> > ContextArray = ChildJSONObject->GetArrayField( TAG_KEYCOLLECTION);
-
-				for(TArray< TSharedPtr< FJsonValue > >::TConstIterator ContextIter( ContextArray.CreateConstIterator() ); ContextIter && bConvertSuccess; ++ContextIter)
+				for (const UE::Json::FValue& ContextEntry : *ContextArray)
 				{
-					const TSharedPtr< FJsonValue > ContextEntry = *ContextIter;
-					const TSharedPtr< FJsonObject > ContextJSONObject = ContextEntry->AsObject();
+					if (!ContextEntry.IsObject())
+					{
+						return false;
+					}
 
-					if( ContextJSONObject->HasTypedField< EJson::String >( TAG_KEY ) )
+					UE::Json::FConstObject ContextJSONObject = ContextEntry.GetObject();
+
+					if (TOptional<FStringView> Key = UE::Json::GetStringField(ContextJSONObject, *TAG_KEY))
 					{
-						const FString Key = ContextJSONObject->GetStringField( TAG_KEY );
-						const FString SourceLocation = ContextJSONObject->HasField( TAG_PATH ) ? ContextJSONObject->GetStringField( TAG_PATH ) : FString();
+						TOptional<FStringView> SourceLocation = UE::Json::GetStringField(ContextJSONObject, *TAG_PATH);
 
 						FManifestContext CommandContext;
-						CommandContext.Key = Key;
-						CommandContext.SourceLocation = SourceLocation;
+						CommandContext.Key = FString(*Key);
+						CommandContext.SourceLocation = SourceLocation.Get(FString());
 						CommandContext.PlatformName = PlatformName;
 
-						if( ContextJSONObject->HasTypedField< EJson::Boolean >( TAG_OPTIONAL ) )
+						if (TOptional<bool> IsOptionalValue = UE::Json::GetBoolField(ContextJSONObject, *TAG_OPTIONAL))
 						{
-							CommandContext.bIsOptional = ContextJSONObject->GetBoolField( TAG_OPTIONAL );
+							CommandContext.bIsOptional = *IsOptionalValue;
 						}
 
-						if( ContextJSONObject->HasTypedField< EJson::Object >( TAG_METADATA ) )
+						if (TOptional<UE::Json::FConstObject> MetaDataJSONObject = UE::Json::GetObjectField(ContextJSONObject, *TAG_METADATA))
 						{
-							const TSharedPtr< FJsonObject > MetaDataJSONObject = ContextJSONObject->GetObjectField( TAG_METADATA );
-
-							if( MetaDataJSONObject->HasTypedField< EJson::Object >( TAG_METADATA_INFO ) )
+							if (TOptional<UE::Json::FConstObject> MetaDataInfoJSONObject = UE::Json::GetObjectField(*MetaDataJSONObject, *TAG_METADATA_INFO))
 							{
-								const TSharedPtr< FJsonObject > MetaDataInfoJSONObject = MetaDataJSONObject->GetObjectField( TAG_METADATA_INFO );
-
 								TSharedPtr< FLocMetadataObject > MetadataNode;
-								FJsonInternationalizationMetaDataSerializer::DeserializeMetadata( MetaDataInfoJSONObject.ToSharedRef(), MetadataNode );
-								if( MetadataNode.IsValid() )
+								FJsonInternationalizationMetaDataSerializer::DeserializeMetadata(*MetaDataInfoJSONObject, MetadataNode);
+								if (MetadataNode)
 								{
 									CommandContext.InfoMetadataObj = MetadataNode;
 								}
 							}
 
-							if( MetaDataJSONObject->HasTypedField< EJson::Object >( TAG_METADATA_KEY ) )
+							if (TOptional<UE::Json::FConstObject> MetaDataKeyJSONObject = UE::Json::GetObjectField(*MetaDataJSONObject, *TAG_METADATA_KEY))
 							{
-								const TSharedPtr< FJsonObject > MetaDataKeyJSONObject = MetaDataJSONObject->GetObjectField( TAG_METADATA_KEY );
-
 								TSharedPtr< FLocMetadataObject > MetadataNode;
-								FJsonInternationalizationMetaDataSerializer::DeserializeMetadata( MetaDataKeyJSONObject.ToSharedRef(), MetadataNode );
-								if( MetadataNode.IsValid() )
+								FJsonInternationalizationMetaDataSerializer::DeserializeMetadata(*MetaDataKeyJSONObject, MetadataNode);
+								if (MetadataNode)
 								{
 									CommandContext.KeyMetadataObj = MetadataNode;
 								}
 							}
 						}
-						bool bAddSuccessful = Manifest->AddSource( AccumulatedNamespace, Source, CommandContext );
-						if(!bAddSuccessful)
+
+						bool bAddSuccessful = Manifest->AddSource(AccumulatedNamespace, Source, CommandContext);
+						if (!bAddSuccessful)
 						{
-							UE_LOG( LogInternationalizationManifestSerializer, Warning,TEXT("Could not add JSON entry to the Internationalization manifest: Namespace:%s SourceText:%s SourceData:%s"), 
-								*AccumulatedNamespace, 
-								*Source.Text, 
-								*FJsonInternationalizationMetaDataSerializer::MetadataToString(Source.MetadataObj) );
+							UE_LOG(LogInternationalizationManifestSerializer, Warning, TEXT("Could not add JSON entry to the Internationalization manifest: Namespace:%s SourceText:%s SourceData:%s"),
+								*AccumulatedNamespace,
+								*Source.Text,
+								*FJsonInternationalizationMetaDataSerializer::MetadataToString(Source.MetadataObj));
 						}
 					}
 					else
 					{
 						//We found a context entry that is missing a identifier/key or a path
-						bConvertSuccess = false;
-						break;
+						return false;
 					}
 
 				}
@@ -304,31 +300,30 @@ bool FJsonInternationalizationManifestSerializer::JsonObjToManifest( TSharedRef<
 			else
 			{
 				// We have an entry that is missing a key/context collection or default text entry.
-				bConvertSuccess = false;
-				break;
+				return false;
 			}
 
 		}
 	}
 
-	if( bConvertSuccess && InJsonObj->HasField( TAG_SUBNAMESPACES ) )
+	if (TOptional<UE::Json::FConstArray> SubnamespaceArray = UE::Json::GetArrayField(InJsonObj, *TAG_SUBNAMESPACES))
 	{
-		const TArray< TSharedPtr<FJsonValue> > SubnamespaceArray = InJsonObj->GetArrayField( TAG_SUBNAMESPACES );
-
-		for(TArray< TSharedPtr< FJsonValue > >::TConstIterator SubnamespaceIter( SubnamespaceArray.CreateConstIterator() ); SubnamespaceIter; ++SubnamespaceIter )
+		for (const UE::Json::FValue& SubnamespaceEntry : *SubnamespaceArray)
 		{
-			const TSharedPtr< FJsonValue >  SubnamespaceEntry = *SubnamespaceIter;
-			const TSharedPtr< FJsonObject > SubnamespaceJSONObject = SubnamespaceEntry->AsObject();
+			if (!SubnamespaceEntry.IsObject())
+			{
+				return false;
+			}
 
-			if( !JsonObjToManifest( SubnamespaceJSONObject.ToSharedRef(), AccumulatedNamespace, Manifest, PlatformName ) )
+			UE::Json::FConstObject SubnamespaceJSONObject = SubnamespaceEntry.GetObject();
+			if (!JsonObjToManifest(SubnamespaceJSONObject, AccumulatedNamespace, Manifest, PlatformName))
 			{
-				bConvertSuccess = false;
-				break;
+				return false;
 			}
 		}
 	}
-	
-	return bConvertSuccess;
+
+	return true;
 }
 
 
@@ -413,116 +408,105 @@ void FJsonInternationalizationManifestSerializer::SortStructuredData( TSharedPtr
 }
 
 
-void FJsonInternationalizationManifestSerializer::StructuredDataToJsonObj( TSharedPtr< const FStructuredEntry > InElement, TSharedRef< FJsonObject > JsonObj )
+void FJsonInternationalizationManifestSerializer::StructuredDataToJsonObj( TSharedPtr< const FStructuredEntry > InElement, UE::Json::FObject JsonObj, UE::Json::FAllocator& Allocator )
 {
+	JsonObj.AddMember(UE::Json::MakeStringRef(TAG_NAMESPACE), UE::Json::MakeStringValue(InElement->Namespace, Allocator), Allocator);
 
-	JsonObj->SetStringField( TAG_NAMESPACE, InElement->Namespace );
-
-	TArray< TSharedPtr< FJsonValue > > NamespaceArray;
-	TArray< TSharedPtr< FJsonValue > > EntryArray;
-
-	//Write namespace content entries
-	for( TArray< TSharedPtr< FManifestEntry > >::TConstIterator Iter( InElement->ManifestEntries ); Iter; ++Iter )
+	// Write namespace content entries
+	UE::Json::FValue EntryJSONArray(rapidjson::kArrayType);
+	for (const TSharedPtr<FManifestEntry>& Entry : InElement->ManifestEntries)
 	{
-		const TSharedPtr< FManifestEntry > Entry = *Iter;
-		TSharedPtr< FJsonObject > EntryNode = MakeShareable( new FJsonObject );
-
-		TSharedPtr<FJsonObject> SourceNode;
+		UE::Json::FValue EntryJSONObject(rapidjson::kObjectType);
 
-		if( Entry->Source.MetadataObj.IsValid() )
-		{
-			FJsonInternationalizationMetaDataSerializer::SerializeMetadata( Entry->Source.MetadataObj.ToSharedRef(), SourceNode );
-		}
-		
-		if( !SourceNode.IsValid() )
 		{
-			SourceNode = MakeShareable( new FJsonObject );
-		}
+			UE::Json::FValue SourceJSONObject(rapidjson::kObjectType);
 
-		// Add escapes for special chars - doesn't do backslash
-		{
-			FString ProcessedText = Entry->Source.Text;
-			SourceNode->SetStringField(TAG_SOURCE_TEXT, ProcessedText);
-		}
+			if (Entry->Source.MetadataObj)
+			{
+				FJsonInternationalizationMetaDataSerializer::SerializeMetadata(Entry->Source.MetadataObj.ToSharedRef(), SourceJSONObject, Allocator);
+			}
 
-		EntryNode->SetObjectField( TAG_SOURCE, SourceNode );
+			SourceJSONObject.AddMember(UE::Json::MakeStringRef(TAG_SOURCE_TEXT), UE::Json::MakeStringValue(Entry->Source.Text, Allocator), Allocator);
 
-		TArray< TSharedPtr< FJsonValue > > KeyArray;
+			EntryJSONObject.AddMember(UE::Json::MakeStringRef(TAG_SOURCE), MoveTemp(SourceJSONObject), Allocator);
+		}
 
-		for(auto ContextIter = Entry->Contexts.CreateConstIterator(); ContextIter; ++ContextIter)
-		{
-			const FManifestContext& AContext = *ContextIter;
+		UE::Json::FValue KeyJSONArray(rapidjson::kArrayType);
 
+		for (const FManifestContext& AContext : Entry->Contexts)
+		{
 			FString ProcessedText = AContext.SourceLocation;
 			ProcessedText.ReplaceInline( TEXT("\\"), TEXT("/"));
 			ProcessedText.ReplaceInline( *FPaths::RootDir(), TEXT("/"));
 
-			TSharedPtr<FJsonObject> KeyNode = MakeShareable( new FJsonObject );
-			KeyNode->SetStringField( TAG_KEY, AContext.Key.GetString() );
-			KeyNode->SetStringField( TAG_PATH, ProcessedText );
+			UE::Json::FValue KeyJSONObject(rapidjson::kObjectType);
+
+			KeyJSONObject.AddMember(UE::Json::MakeStringRef(TAG_KEY), UE::Json::MakeStringValue(AContext.Key.GetString(), Allocator), Allocator);
+			KeyJSONObject.AddMember(UE::Json::MakeStringRef(TAG_PATH), UE::Json::MakeStringValue(ProcessedText, Allocator), Allocator);
 
 			// We only add the optional field if it is true, it is assumed to be false otherwise.
-			if( AContext.bIsOptional == true )
+			if (AContext.bIsOptional)
 			{
-				KeyNode->SetBoolField( TAG_OPTIONAL, AContext.bIsOptional);
+				KeyJSONObject.AddMember(UE::Json::MakeStringRef(TAG_OPTIONAL), UE::Json::FValue(AContext.bIsOptional), Allocator);
 			}
 
-			TSharedPtr< FJsonObject > MetaDataNode = MakeShareable( new FJsonObject );
-
-			if( AContext.InfoMetadataObj.IsValid() )
 			{
-				TSharedPtr< FJsonObject > InfoDataNode;
-				FJsonInternationalizationMetaDataSerializer::SerializeMetadata( AContext.InfoMetadataObj.ToSharedRef(), InfoDataNode );
-				if( InfoDataNode.IsValid() )
+				UE::Json::FValue MetaDataJSONObject(rapidjson::kObjectType);
+
+				if (AContext.InfoMetadataObj)
 				{
-					MetaDataNode->SetObjectField( TAG_METADATA_INFO, InfoDataNode );
+					UE::Json::FValue InfoDataJSONObject(rapidjson::kObjectType);
+					FJsonInternationalizationMetaDataSerializer::SerializeMetadata(AContext.InfoMetadataObj.ToSharedRef(), InfoDataJSONObject, Allocator);
+					if (InfoDataJSONObject.MemberCount() > 0)
+					{
+						MetaDataJSONObject.AddMember(UE::Json::MakeStringRef(TAG_METADATA_INFO), MoveTemp(InfoDataJSONObject), Allocator);
+					}
 				}
-			}
 
-			if( AContext.KeyMetadataObj.IsValid() )
-			{
-				TSharedPtr< FJsonObject > KeyDataNode;
-				FJsonInternationalizationMetaDataSerializer::SerializeMetadata( AContext.KeyMetadataObj.ToSharedRef(), KeyDataNode );
-				if( KeyDataNode.IsValid() )
+				if (AContext.KeyMetadataObj)
 				{
-					MetaDataNode->SetObjectField( TAG_METADATA_KEY, KeyDataNode );
+					UE::Json::FValue KeyDataJSONObject(rapidjson::kObjectType);
+					FJsonInternationalizationMetaDataSerializer::SerializeMetadata(AContext.KeyMetadataObj.ToSharedRef(), KeyDataJSONObject, Allocator);
+					if (KeyDataJSONObject.MemberCount() > 0)
+					{
+						MetaDataJSONObject.AddMember(UE::Json::MakeStringRef(TAG_METADATA_KEY), MoveTemp(KeyDataJSONObject), Allocator);
+					}
 				}
-			}
 
-			if( MetaDataNode->Values.Num() > 0 )
-			{
-				KeyNode->SetObjectField( TAG_METADATA, MetaDataNode );
+				if (MetaDataJSONObject.MemberCount() > 0)
+				{
+					KeyJSONObject.AddMember(UE::Json::MakeStringRef(TAG_METADATA), MoveTemp(MetaDataJSONObject), Allocator);
+				}
 			}
 			
-			KeyArray.Add( MakeShareable( new FJsonValueObject( KeyNode ) ) );
+			KeyJSONArray.PushBack(MoveTemp(KeyJSONObject), Allocator);
 		}
 
-		EntryNode->SetArrayField( TAG_KEYCOLLECTION, KeyArray);
+		EntryJSONObject.AddMember(UE::Json::MakeStringRef(TAG_KEYCOLLECTION), MoveTemp(KeyJSONArray), Allocator);
 
-		EntryArray.Add( MakeShareable( new FJsonValueObject( EntryNode ) ) );
+		EntryJSONArray.PushBack(MoveTemp(EntryJSONObject), Allocator);
 	}
 
-	//Write the subnamespaces
-	for( auto Iter = InElement->SubNamespaces.CreateConstIterator(); Iter; ++Iter )
+	// Write the subnamespaces
+	UE::Json::FValue NamespaceJSONArray(rapidjson::kArrayType);
+	for (const TTuple<FString, TSharedPtr<FStructuredEntry>>& SubElementPair : InElement->SubNamespaces)
 	{
-		const TSharedPtr< FStructuredEntry > SubElement = Iter->Value;
-
-		if(SubElement.IsValid())
+		if (SubElementPair.Value)
 		{
-			TSharedRef< FJsonObject > SubObject = MakeShareable( new FJsonObject );
-			StructuredDataToJsonObj( SubElement, SubObject );
+			UE::Json::FValue SubJSONObject(rapidjson::kObjectType);
+			StructuredDataToJsonObj(SubElementPair.Value, SubJSONObject.GetObject(), Allocator);
 
-			NamespaceArray.Add( MakeShareable( new FJsonValueObject( SubObject ) ) );
+			NamespaceJSONArray.PushBack(MoveTemp(SubJSONObject), Allocator);
 		}
 	}
 
-	if( EntryArray.Num() > 0 )
+	if (EntryJSONArray.Size() > 0)
 	{
-		JsonObj->SetArrayField( TAG_CHILDREN, EntryArray );
+		JsonObj.AddMember(UE::Json::MakeStringRef(TAG_CHILDREN), MoveTemp(EntryJSONArray), Allocator);
 	}
 
-	if( NamespaceArray.Num() > 0 )
+	if (NamespaceJSONArray.Size() > 0)
 	{
-		JsonObj->SetArrayField( TAG_SUBNAMESPACES, NamespaceArray );
+		JsonObj.AddMember(UE::Json::MakeStringRef(TAG_SUBNAMESPACES), MoveTemp(NamespaceJSONArray), Allocator);
 	}
 }
diff --git a/Engine/Source/Developer/Localization/Private/Serialization/JsonInternationalizationMetadataSerializer.cpp b/Engine/Source/Developer/Localization/Private/Serialization/JsonInternationalizationMetadataSerializer.cpp
--- a/Engine/Source/Developer/Localization/Private/Serialization/JsonInternationalizationMetadataSerializer.cpp
+++ b/Engine/Source/Developer/Localization/Private/Serialization/JsonInternationalizationMetadataSerializer.cpp
@@ -1,7 +1,6 @@
 // Copyright Epic Games, Inc. All Rights Reserved.
 
 #include "Serialization/JsonInternationalizationMetadataSerializer.h"
-#include "Serialization/JsonSerializer.h"
 
 
 struct FCompareLocMetadataValue
@@ -21,173 +20,166 @@ struct FCompareLocMetadataValue
 };
 
 
-TSharedPtr<FLocMetadataValue> JSonValueToLocMetaDataValue( const TSharedRef< FJsonValue > JsonValue )
+TSharedPtr<FLocMetadataValue> JsonValueToLocMetaDataValue(const UE::Json::FValue& JsonValue)
 {
-	switch( JsonValue->Type )
+	switch (JsonValue.GetType())
 	{
-	case EJson::Boolean:
+		case rapidjson::kFalseType:
+		case rapidjson::kTrueType:
 		{
-			return MakeShareable( new FLocMetadataValueBoolean( JsonValue->AsBool() ) );
+			return MakeShared<FLocMetadataValueBoolean>(JsonValue.GetBool());
 		}
-	case EJson::String:
+
+		case rapidjson::kStringType:
 		{
-			return MakeShareable( new FLocMetadataValueString( JsonValue->AsString() ) );
+			return MakeShared<FLocMetadataValueString>(FString(FStringView(JsonValue.GetString(), JsonValue.GetStringLength())));
 		}
-	case EJson::Array:
+
+		case rapidjson::kArrayType:
 		{
-			TArray< TSharedPtr< FLocMetadataValue > > MetadataArray;
-			TArray< TSharedPtr< FJsonValue > > JsonArray = JsonValue->AsArray();
-			for( auto ArrayIter = JsonArray.CreateConstIterator(); ArrayIter; ++ArrayIter )
+			TArray<TSharedPtr<FLocMetadataValue>> MetadataArray;
+			UE::Json::FConstArray JsonArray = JsonValue.GetArray();
+			for (const UE::Json::FValue& Item : JsonArray)
 			{
-				const TSharedPtr< FJsonValue >& Item = *ArrayIter;
-				if( Item.IsValid() )
+				if (TSharedPtr<FLocMetadataValue> Entry = JsonValueToLocMetaDataValue(Item))
 				{
-					TSharedPtr< FLocMetadataValue > Entry = JSonValueToLocMetaDataValue( Item.ToSharedRef() );
-					if( Entry.IsValid() )
-					{
-						MetadataArray.Add( Entry );
-					}
+					MetadataArray.Add(Entry);
 				}
 			}
-			if( MetadataArray.Num() > 0 )
+			if (MetadataArray.Num() > 0)
 			{
-				return MakeShareable( new FLocMetadataValueArray( MetadataArray ) );
+				return MakeShared<FLocMetadataValueArray>(MoveTemp(MetadataArray));
 			}
 		}
 		break;
-	case EJson::Object:
+
+		case rapidjson::kObjectType:
 		{
-			TSharedPtr< FLocMetadataObject > MetadataSubObject = MakeShareable( new FLocMetadataObject );
-			TSharedPtr< FJsonObject > JsonObject = JsonValue->AsObject();
-			for( auto ValueIter = JsonObject->Values.CreateConstIterator(); ValueIter; ++ValueIter )
+			TSharedPtr<FLocMetadataObject> MetadataSubObject = MakeShared<FLocMetadataObject>();
+			UE::Json::FConstObject JsonObject = JsonValue.GetObject();
+			for (const UE::Json::FMember& ValueIter : JsonObject)
 			{
-				const FString Name = (*ValueIter).Key;
-				TSharedPtr< FJsonValue > Value = (*ValueIter).Value;
-				if( Value.IsValid() )
+				if (TSharedPtr<FLocMetadataValue> MetadataValue = JsonValueToLocMetaDataValue(ValueIter.value))
 				{
-					TSharedPtr< FLocMetadataValue > MetadataValue = JSonValueToLocMetaDataValue( Value.ToSharedRef() );
-					if( MetadataValue.IsValid() )
-					{
-						MetadataSubObject->SetField( Name, MetadataValue );
-					}
+					MetadataSubObject->SetField(FString(FStringView(ValueIter.name.GetString(), ValueIter.name.GetStringLength())), MetadataValue);
 				}
 			}
-			return MakeShareable( new FLocMetadataValueObject(MetadataSubObject) );
+			return MakeShared<FLocMetadataValueObject>(MoveTemp(MetadataSubObject));
 		}
-	default:
+
+		default:
 		{
 			// At the moment we do not support all the json types as metadata.  In the future these types will be handled in a way that they can be stored in an unprocessed way.
 		}
 		break;
 	}
 
-	return NULL;
+	return nullptr;
 }
 
 
-TSharedPtr<FJsonValue> LocMetaDataValueToJsonValue( const TSharedRef< FLocMetadataValue > MetadataValue )
+TOptional<UE::Json::FValue> LocMetaDataValueToJsonValue(const TSharedRef< FLocMetadataValue > MetadataValue, UE::Json::FValue* OutTopLevelObjectValuePtr, UE::Json::FAllocator& Allocator)
 {
-	switch( MetadataValue->GetType() )
+	switch (MetadataValue->GetType())
 	{
-	case ELocMetadataType::Boolean:
+		case ELocMetadataType::Boolean:
 		{
-			return MakeShareable( new FJsonValueBoolean( MetadataValue->AsBool() ) );
+			return UE::Json::FValue(MetadataValue->AsBool());
 		}
 
-	case ELocMetadataType::String:
+		case ELocMetadataType::String:
 		{
-			return MakeShareable( new FJsonValueString( MetadataValue->AsString() ) );
+			return UE::Json::MakeStringValue(MetadataValue->AsString(), Allocator);
 		}
 
-	case ELocMetadataType::Array:
+		case ELocMetadataType::Array:
 		{
-			TArray< TSharedPtr<FJsonValue> > JsonArrayVals;
 			TArray< TSharedPtr< FLocMetadataValue > > MetaDataArray = MetadataValue->AsArray();
-			
-			MetaDataArray.Sort( FCompareLocMetadataValue() );
-			for( auto ArrayIter = MetaDataArray.CreateConstIterator(); ArrayIter; ++ArrayIter )
+			MetaDataArray.Sort(FCompareLocMetadataValue());
+
+			UE::Json::FValue JsonArrayVals(rapidjson::kArrayType);
+			for (const TSharedPtr<FLocMetadataValue>& Item : MetaDataArray)
 			{
-				const TSharedPtr<FLocMetadataValue>& Item = *ArrayIter;
-				if( Item.IsValid() )
+				if (Item)
 				{
-					TSharedPtr<FJsonValue> Entry = LocMetaDataValueToJsonValue( Item.ToSharedRef() );
-					if( Entry.IsValid() )
+					if (TOptional<UE::Json::FValue> Entry = LocMetaDataValueToJsonValue(Item.ToSharedRef(), nullptr, Allocator))
 					{
-						JsonArrayVals.Add( Entry );
+						JsonArrayVals.PushBack(MoveTemp(*Entry), Allocator);
 					}
 				}
 			}
 
-			if( JsonArrayVals.Num() > 0 )
+			if (JsonArrayVals.Size() > 0)
 			{
-				return MakeShareable( new FJsonValueArray( JsonArrayVals ) );
+				return JsonArrayVals;
 			}
 		}
 		break;
 
-	case ELocMetadataType::Object:
+		case ELocMetadataType::Object:
 		{
-			TSharedPtr< FJsonObject > JsonSubObject = MakeShareable( new FJsonObject );
-			TSharedPtr< FLocMetadataObject > MetaDataObject = MetadataValue->AsObject();
-			for( auto ValueIter = MetaDataObject->Values.CreateConstIterator(); ValueIter; ++ValueIter )
+			TMap<FString, TSharedPtr<FLocMetadataValue>> MetaDataMap = MetadataValue->AsObject()->Values;
+			// Sorting by key is probably sufficient for now but ideally we would sort the resulting json object using
+			//  the same logic that is in the FLocMetadata < operator
+			MetaDataMap.KeySort(TLess<FString>());
+
+			TOptional<UE::Json::FValue> JsonSubObject;
+			if (!OutTopLevelObjectValuePtr)
+			{
+				JsonSubObject.Emplace(rapidjson::kObjectType);
+			}
+
+			UE::Json::FValue& JsonObject = OutTopLevelObjectValuePtr ? *OutTopLevelObjectValuePtr : JsonSubObject.GetValue();
+			check(JsonObject.IsObject());
+			for (const TTuple<FString, TSharedPtr<FLocMetadataValue>>& MetaDataPair : MetaDataMap)
 			{
-				const FString Name = (*ValueIter).Key;
-				TSharedPtr< FLocMetadataValue > Value = (*ValueIter).Value;
-				if( Value.IsValid() )
+				if (MetaDataPair.Value)
 				{
-					TSharedPtr< FJsonValue > JsonValue = LocMetaDataValueToJsonValue( Value.ToSharedRef() );
-					if( JsonValue.IsValid() )
+					if (TOptional<UE::Json::FValue> JsonValue = LocMetaDataValueToJsonValue(MetaDataPair.Value.ToSharedRef(), nullptr, Allocator))
 					{
-						JsonSubObject->SetField( Name, JsonValue );
+						JsonObject.AddMember(UE::Json::MakeStringValue(MetaDataPair.Key, Allocator), MoveTemp(*JsonValue), Allocator);
 					}
 				}
 			}
 
-			// Sorting by key is probably sufficient for now but ideally we would sort the resulting json object using
-			//  the same logic that is in the FLocMetadata < operator
-			JsonSubObject->Values.KeySort( TLess<FString>() );
-
-			return MakeShareable( new FJsonValueObject(JsonSubObject) );
+			return JsonSubObject;
 		}
 
-	default:
-		break;
+		default:
+			break;
 	}
-	return NULL;
+
+	return {};
 }
 
 
-void FJsonInternationalizationMetaDataSerializer::DeserializeMetadata( const TSharedRef< FJsonObject > JsonObj, TSharedPtr< FLocMetadataObject >& OutMetaDataObj )
+void FJsonInternationalizationMetaDataSerializer::DeserializeMetadata( UE::Json::FConstObject JsonObj, TSharedPtr< FLocMetadataObject >& OutMetaDataObj )
 {
-	OutMetaDataObj = JSonValueToLocMetaDataValue( MakeShareable( new FJsonValueObject( JsonObj ) ) )->AsObject();
+	OutMetaDataObj = JsonValueToLocMetaDataValue(JsonObj)->AsObject();
 }
 
 
-void FJsonInternationalizationMetaDataSerializer::SerializeMetadata( const TSharedRef< FLocMetadataObject > MetaData, TSharedPtr< FJsonObject >& OutJsonObj )
+void FJsonInternationalizationMetaDataSerializer::SerializeMetadata( const TSharedRef< FLocMetadataObject > MetaData, UE::Json::FValue& OutJsonObj, UE::Json::FAllocator& Allocator )
 {
-	OutJsonObj = LocMetaDataValueToJsonValue( MakeShareable( new FLocMetadataValueObject( MetaData ) ) )->AsObject();
+	check(OutJsonObj.IsObject());
+	LocMetaDataValueToJsonValue(MakeShared<FLocMetadataValueObject>(MetaData), &OutJsonObj, Allocator);
 }
 
 
 FString FJsonInternationalizationMetaDataSerializer::MetadataToString( const TSharedPtr<FLocMetadataObject> Metadata )
 {
-	FString StringMetadata = TEXT("");
-	if( Metadata.IsValid() )
+	FString StringMetadata;
+	if (Metadata)
 	{
-		TSharedPtr< FJsonObject > JsonKeyMetadata;
-		FJsonInternationalizationMetaDataSerializer::SerializeMetadata( Metadata.ToSharedRef(), JsonKeyMetadata );
+		UE::Json::FDocument Document(rapidjson::kObjectType);
 
-		if( JsonKeyMetadata.IsValid() )
+		FJsonInternationalizationMetaDataSerializer::SerializeMetadata(Metadata.ToSharedRef(), Document, Document.GetAllocator());
+		if (Document.MemberCount() > 0)
 		{
-			JsonKeyMetadata->Values.KeySort( TLess<FString>() );
-
-			TSharedRef< TJsonWriter<> > Writer = TJsonWriterFactory<>::Create( &StringMetadata );
-			FJsonSerializer::Serialize( JsonKeyMetadata.ToSharedRef(), Writer );
-			Writer->Close();
-
-			StringMetadata.ReplaceInline( TEXT("\t"), TEXT(" ") );
-			StringMetadata.ReplaceInline( TEXT("\r\n"), TEXT(" ") );
-			StringMetadata.ReplaceInline( TEXT("\n"), TEXT(" ") );
+			StringMetadata = UE::Json::WritePretty(Document);
+			StringMetadata.ReplaceInline(TEXT("\t"), TEXT(" "));
+			StringMetadata.ReplaceInline(TEXT("\r\n"), TEXT(" "));
+			StringMetadata.ReplaceInline(TEXT("\n"), TEXT(" "));
 		}
 	}
 	return StringMetadata;
diff --git a/Engine/Source/Developer/Localization/Public/LocTextHelper.h b/Engine/Source/Developer/Localization/Public/LocTextHelper.h
--- a/Engine/Source/Developer/Localization/Public/LocTextHelper.h
+++ b/Engine/Source/Developer/Localization/Public/LocTextHelper.h
@@ -26,7 +26,6 @@ class FArchiveEntry;
 class FInternationalizationArchive;
 class FLocMetadataObject;
 class FText;
-class FJsonObject;
 
 /** Flags controlling the behavior used when loading manifests and archives into FLocTextHelper */
 enum class ELocTextHelperLoadFlags : uint8
@@ -382,13 +381,6 @@ public:
 	 */
 	UE_API bool SaveManifest(FText* OutError = nullptr) const;
 
-
-	/**
-	 * Attempt to serialize the entire manifest to an input json object
-	 * @return True if the serialization succeeds false otherwise.
-	 */
-	UE_API bool SerializeManifestToJson(TSharedRef<FJsonObject> JsonObject);
-
 	/**
 	 * Attempt to save the manifest file to the given file path.
 	 *
diff --git a/Engine/Source/Developer/Localization/Public/Serialization/JsonInternationalizationArchiveSerializer.h b/Engine/Source/Developer/Localization/Public/Serialization/JsonInternationalizationArchiveSerializer.h
--- a/Engine/Source/Developer/Localization/Public/Serialization/JsonInternationalizationArchiveSerializer.h
+++ b/Engine/Source/Developer/Localization/Public/Serialization/JsonInternationalizationArchiveSerializer.h
@@ -5,11 +5,10 @@
 #include "Containers/Array.h"
 #include "Containers/Map.h"
 #include "Containers/UnrealString.h"
-#include "CoreMinimal.h"
-#include "Dom/JsonObject.h"
 #include "Internationalization/InternationalizationArchive.h"
 #include "Internationalization/InternationalizationManifest.h"
 #include "Internationalization/LocKeyFuncs.h"
+#include "JsonUtils/RapidJsonUtils.h"
 #include "Templates/SharedPointer.h"
 
 #define UE_API LOCALIZATION_API
@@ -17,7 +16,6 @@
 class FArchiveEntry;
 class FInternationalizationArchive;
 class FInternationalizationManifest;
-class FJsonObject;
 
 /**
  * Implements a serializer that serializes to and from Json encoded data.
@@ -36,19 +34,7 @@ public:
 	 *
 	 * @return true if deserialization was successful, false otherwise.
 	 */
-	static UE_API bool DeserializeArchive(const FString& InStr, TSharedRef<FInternationalizationArchive> InArchive, TSharedPtr<const FInternationalizationManifest> InManifest, TSharedPtr<const FInternationalizationArchive> InNativeArchive);
-
-	/**
-	 * Deserializes an archive from a JSON object.
-	 *
-	 * @param InJsonObj			The JSON object to serialize from.
-	 * @param InArchive			The archive to populate from the JSON data.
-	 * @param InManifest		The manifest associated with the archive. May be null, but you won't be able to load archives with a version < FInternationalizationArchive::EFormatVersion::AddedKeys.
-	 * @param InNativeArchive	The native archive associated with the archive. May be null.
-	 *
-	 * @return true if deserialization was successful, false otherwise.
-	 */
-	static UE_API bool DeserializeArchive(TSharedRef<FJsonObject> InJsonObj, TSharedRef<FInternationalizationArchive> InArchive, TSharedPtr<const FInternationalizationManifest> InManifest, TSharedPtr<const FInternationalizationArchive> InNativeArchive);
+	static UE_API bool DeserializeArchiveFromString(const FString& InStr, TSharedRef<FInternationalizationArchive> InArchive, TSharedPtr<const FInternationalizationManifest> InManifest, TSharedPtr<const FInternationalizationArchive> InNativeArchive);
 
 	/**
 	 * Deserializes an archive from a JSON file.
@@ -62,16 +48,6 @@ public:
 	 */
 	static UE_API bool DeserializeArchiveFromFile(const FString& InJsonFile, TSharedRef<FInternationalizationArchive> InArchive, TSharedPtr<const FInternationalizationManifest> InManifest, TSharedPtr<const FInternationalizationArchive> InNativeArchive);
 
-	/**
-	 * Serializes an archive to a JSON object.
-	 *
-	 * @param InArchive			The archive data to serialize.
-	 * @param InJsonObj			The JSON object to serialize into.
-	 *
-	 * @return true if serialization was successful, false otherwise.
-	 */
-	static UE_API bool SerializeArchive(TSharedRef<const FInternationalizationArchive> InArchive, TSharedRef<FJsonObject> InJsonObj);
-
 	/**
 	 * Serializes an archive to a JSON string.
 	 *
@@ -80,7 +56,7 @@ public:
 	 *
 	 * @return true if serialization was successful, false otherwise.
 	 */
-	static UE_API bool SerializeArchive(TSharedRef<const FInternationalizationArchive> InArchive, FString& Str);
+	static UE_API bool SerializeArchiveToString(TSharedRef<const FInternationalizationArchive> InArchive, FString& Str);
 
 	/**
 	* Serializes an archive to a JSON string.
@@ -92,7 +68,7 @@ public:
 	*/
 	static UE_API bool SerializeArchiveToFile(TSharedRef<const FInternationalizationArchive> InArchive, const FString& InJsonFile);
 
-protected:
+private:
 
 	/**
 	 * Used to arrange Internationalization archive data in a hierarchy based on namespace prior to json serialization.
@@ -119,7 +95,7 @@ protected:
 	 *
 	 * @return true if deserialization was successful, false otherwise.
 	 */
-	static UE_API bool DeserializeInternal(TSharedRef<FJsonObject> InJsonObj, TSharedRef<FInternationalizationArchive> InArchive, TSharedPtr<const FInternationalizationManifest> InManifest, TSharedPtr<const FInternationalizationArchive> InNativeArchive);
+	static bool DeserializeInternal(UE::Json::FConstObject InJsonObj, TSharedRef<FInternationalizationArchive> InArchive, TSharedPtr<const FInternationalizationManifest> InManifest, TSharedPtr<const FInternationalizationArchive> InNativeArchive);
 
 	/**
 	 * Convert a Internationalization archive to a JSON object.
@@ -128,7 +104,7 @@ protected:
 	 * @param JsonObj The Json object that will store the data.
 	 * @return true if serialization was successful, false otherwise.
 	 */
-	static UE_API bool SerializeInternal(TSharedRef<const FInternationalizationArchive> InArchive, TSharedRef<FJsonObject> JsonObj);
+	static bool SerializeInternal(TSharedRef<const FInternationalizationArchive> InArchive, UE::Json::FObject JsonObj, UE::Json::FAllocator& Allocator);
 
 	/**
 	 * Recursive function that will traverse the JSON object and populate an archive.
@@ -141,7 +117,7 @@ protected:
 	 *
 	 * @return true if successful, false otherwise.
 	 */
-	static UE_API bool JsonObjToArchive(TSharedRef<FJsonObject> InJsonObj, const FString& ParentNamespace, TSharedRef<FInternationalizationArchive> InArchive, TSharedPtr<const FInternationalizationManifest> InManifest, TSharedPtr<const FInternationalizationArchive> InNativeArchive);
+	static bool JsonObjToArchive(UE::Json::FConstObject InJsonObj, const FString& ParentNamespace, TSharedRef<FInternationalizationArchive> InArchive, TSharedPtr<const FInternationalizationManifest> InManifest, TSharedPtr<const FInternationalizationArchive> InNativeArchive);
 
 	/**
 	 * Takes a Internationalization archive and arranges the data into a hierarchy based on namespace.
@@ -149,14 +125,14 @@ protected:
 	 * @param InArchive The Internationalization archive.
 	 * @param RootElement The root element of the structured data.
 	 */
-	static UE_API void GenerateStructuredData( TSharedRef< const FInternationalizationArchive > InArchive, TSharedPtr< FStructuredArchiveEntry > RootElement );
+	static void GenerateStructuredData( TSharedRef< const FInternationalizationArchive > InArchive, TSharedPtr< FStructuredArchiveEntry > RootElement );
 
 	/**
 	 * Goes through the structured, hierarchy based, archive data and does a non-culture specific sort on namespaces and default text.
 	 *
 	 * @param RootElement The root element of the structured data.
 	 */
-	static UE_API void SortStructuredData( TSharedPtr< FStructuredArchiveEntry > InElement );
+	static void SortStructuredData( TSharedPtr< FStructuredArchiveEntry > InElement );
 
 	/**
 	 * Populates a JSON object from Internationalization archive data that has been structured based on namespace.
@@ -164,25 +140,23 @@ protected:
 	 * @param InElement Internationalization archive data structured based on namespace.
 	 * @param JsonObj JSON object to be populated.
 	 */
-	static UE_API void StructuredDataToJsonObj(TSharedPtr< const FStructuredArchiveEntry > InElement, TSharedRef< FJsonObject > JsonObj );
-
-public:
-
-	static UE_API const FString TAG_FORMATVERSION;
-	static UE_API const FString TAG_NAMESPACE;
-	static UE_API const FString TAG_KEY;
-	static UE_API const FString TAG_CHILDREN;
-	static UE_API const FString TAG_SUBNAMESPACES;
-	static UE_API const FString TAG_DEPRECATED_DEFAULTTEXT;
-	static UE_API const FString TAG_DEPRECATED_TRANSLATEDTEXT;
-	static UE_API const FString TAG_OPTIONAL;
-	static UE_API const FString TAG_SOURCE;
-	static UE_API const FString TAG_SOURCE_TEXT;
-	static UE_API const FString TAG_TRANSLATION;
-	static UE_API const FString TAG_TRANSLATION_TEXT;
-	static UE_API const FString TAG_METADATA;
-	static UE_API const FString TAG_METADATA_KEY;
-	static UE_API const FString NAMESPACE_DELIMITER;
+	static void StructuredDataToJsonObj(TSharedPtr< const FStructuredArchiveEntry > InElement, UE::Json::FObject JsonObj, UE::Json::FAllocator& Allocator);
+
+	static const FString TAG_FORMATVERSION;
+	static const FString TAG_NAMESPACE;
+	static const FString TAG_KEY;
+	static const FString TAG_CHILDREN;
+	static const FString TAG_SUBNAMESPACES;
+	static const FString TAG_DEPRECATED_DEFAULTTEXT;
+	static const FString TAG_DEPRECATED_TRANSLATEDTEXT;
+	static const FString TAG_OPTIONAL;
+	static const FString TAG_SOURCE;
+	static const FString TAG_SOURCE_TEXT;
+	static const FString TAG_TRANSLATION;
+	static const FString TAG_TRANSLATION_TEXT;
+	static const FString TAG_METADATA;
+	static const FString TAG_METADATA_KEY;
+	static const FString NAMESPACE_DELIMITER;
 };
 
 #undef UE_API
diff --git a/Engine/Source/Developer/Localization/Public/Serialization/JsonInternationalizationManifestSerializer.h b/Engine/Source/Developer/Localization/Public/Serialization/JsonInternationalizationManifestSerializer.h
--- a/Engine/Source/Developer/Localization/Public/Serialization/JsonInternationalizationManifestSerializer.h
+++ b/Engine/Source/Developer/Localization/Public/Serialization/JsonInternationalizationManifestSerializer.h
@@ -5,10 +5,9 @@
 #include "Containers/Array.h"
 #include "Containers/Map.h"
 #include "Containers/UnrealString.h"
-#include "CoreMinimal.h"
-#include "Dom/JsonObject.h"
 #include "Internationalization/InternationalizationManifest.h"
 #include "Internationalization/LocKeyFuncs.h"
+#include "JsonUtils/RapidJsonUtils.h"
 #include "Templates/SharedPointer.h"
 #include "UObject/NameTypes.h"
 
@@ -16,7 +15,6 @@
 
 class FInternationalizationManifest;
 class FManifestEntry;
-class FJsonObject;
 
 /**
  * Implements a serializer that serializes to and from Json encoded data.
@@ -32,16 +30,7 @@ public:
 	 * @param Manifest The populated Internationalization manifest.
 	 * @return true if deserialization was successful, false otherwise.
 	 */
-	static UE_API bool DeserializeManifest( const FString& InStr, TSharedRef< FInternationalizationManifest > Manifest, const FName PlatformName = FName() );
-
-	/**
-	 * Deserializes a Internationalization manifest from a JSON object.
-	 *
-	 * @param InJsonObj The JSON object to serialize from.
-	 * @param Manifest The populated Internationalization manifest.
-	 * @return true if deserialization was successful, false otherwise.
-	 */
-	static UE_API bool DeserializeManifest( TSharedRef< FJsonObject > InJsonObj, TSharedRef< FInternationalizationManifest > Manifest, const FName PlatformName = FName() );
+	static UE_API bool DeserializeManifestFromString( const FString& InStr, TSharedRef< FInternationalizationManifest > Manifest, const FName PlatformName = FName() );
 
 	/**
 	 * Deserializes a Internationalization manifest from a JSON file.
@@ -59,16 +48,7 @@ public:
 	 * @param Str The string to serialize into.
 	 * @return true if serialization was successful, false otherwise.
 	 */
-	static UE_API bool SerializeManifest( TSharedRef< const FInternationalizationManifest > Manifest, FString& Str );
-
-	/**
-	 * Serializes a Internationalization manifest to a JSON object.
-	 *
-	 * @param Manifest The Internationalization manifest data to serialize.
-	 * @param JsonObj The JSON object to serialize into.
-	 * @return true if serialization was successful, false otherwise.
-	 */
-	static UE_API bool SerializeManifest( TSharedRef< const FInternationalizationManifest > Manifest, TSharedRef< FJsonObject > JsonObj );
+	static UE_API bool SerializeManifestToString( TSharedRef< const FInternationalizationManifest > Manifest, FString& Str );
 
 	/**
 	 * Serializes a Internationalization manifest to a JSON file.
@@ -79,7 +59,7 @@ public:
 	 */
 	static UE_API bool SerializeManifestToFile( TSharedRef< const FInternationalizationManifest > Manifest, const FString& InJsonFile );
 
-protected:
+private:
 
 	struct FStructuredEntry
 	{
@@ -101,7 +81,7 @@ protected:
 	 *
 	 * @return true if deserialization was successful, false otherwise.
 	 */
-	static UE_API bool DeserializeInternal( const TSharedRef< FJsonObject > InJsonObj, TSharedRef< FInternationalizationManifest > Manifest, const FName PlatformName );
+	static bool DeserializeInternal( UE::Json::FConstObject InJsonObj, TSharedRef< FInternationalizationManifest > Manifest, const FName PlatformName );
 
 	/**
 	 * Convert a Internationalization manifest to a JSON object.
@@ -111,7 +91,7 @@ protected:
 	 *
 	 * @return true if serialization was successful, false otherwise.
 	 */
-	static UE_API bool SerializeInternal( TSharedRef< const FInternationalizationManifest > InManifest, TSharedRef< FJsonObject > JsonObj );
+	static bool SerializeInternal( TSharedRef< const FInternationalizationManifest > InManifest, UE::Json::FObject JsonObj, UE::Json::FAllocator& Allocator );
 
 	/**
 	 * Recursive function that will traverse the JSON object and populate a Internationalization manifest.
@@ -122,7 +102,7 @@ protected:
 	 *
 	 * @return true if successful, false otherwise.
 	 */
-	static UE_API bool JsonObjToManifest( TSharedRef< FJsonObject > InJsonObj, FString InNamespace, TSharedRef< FInternationalizationManifest > Manifest, const FName PlatformName );
+	static bool JsonObjToManifest( UE::Json::FConstObject InJsonObj, FString InNamespace, TSharedRef< FInternationalizationManifest > Manifest, const FName PlatformName );
 
 	/**
 	 * Takes a Internationalization manifest and arranges the data into a hierarchy based on namespace.
@@ -130,14 +110,14 @@ protected:
 	 * @param InManifest The Internationalization manifest.
 	 * @param RootElement The root element of the structured data.
 	 */
-	static UE_API void GenerateStructuredData( TSharedRef< const FInternationalizationManifest > InManifest, TSharedPtr< FStructuredEntry > RootElement );
+	static void GenerateStructuredData( TSharedRef< const FInternationalizationManifest > InManifest, TSharedPtr< FStructuredEntry > RootElement );
 
 	/**
 	 * Goes through the structured, hierarchy based, manifest data and does a non-culture specific sort on namespaces, default text, and key.
 	 *
 	 * @param RootElement The root element of the structured data.
 	 */
-	static UE_API void SortStructuredData( TSharedPtr< FStructuredEntry > InElement );
+	static void SortStructuredData( TSharedPtr< FStructuredEntry > InElement );
 
 	/**
 	 * Populates a JSON object from Internationalization manifest data that has been structured based on namespace.
@@ -145,26 +125,25 @@ protected:
 	 * @param InElement Internationalization manifest data structured based on namespace.
 	 * @param JsonObj JSON object to be populated.
 	 */
-	static UE_API void StructuredDataToJsonObj( TSharedPtr< const FStructuredEntry> InElement, TSharedRef< FJsonObject > JsonObj );
-
-public:
-
-	static UE_API const FString TAG_FORMATVERSION;
-	static UE_API const FString TAG_NAMESPACE;
-	static UE_API const FString TAG_CHILDREN;
-	static UE_API const FString TAG_SUBNAMESPACES;
-	static UE_API const FString TAG_PATH;
-	static UE_API const FString TAG_OPTIONAL;
-	static UE_API const FString TAG_KEYCOLLECTION;
-	static UE_API const FString TAG_KEY;
-	static UE_API const FString TAG_DEPRECATED_DEFAULTTEXT;
-	static UE_API const FString TAG_SOURCE;
-	static UE_API const FString TAG_SOURCE_TEXT;
-	static UE_API const FString TAG_PLATFORM_NAME;
-	static UE_API const FString TAG_METADATA;
-	static UE_API const FString TAG_METADATA_INFO;
-	static UE_API const FString TAG_METADATA_KEY;
-	static UE_API const FString NAMESPACE_DELIMITER;
+	static void StructuredDataToJsonObj( TSharedPtr< const FStructuredEntry> InElement, UE::Json::FObject JsonObj, UE::Json::FAllocator& Allocator );
+
+
+	static const FString TAG_FORMATVERSION;
+	static const FString TAG_NAMESPACE;
+	static const FString TAG_CHILDREN;
+	static const FString TAG_SUBNAMESPACES;
+	static const FString TAG_PATH;
+	static const FString TAG_OPTIONAL;
+	static const FString TAG_KEYCOLLECTION;
+	static const FString TAG_KEY;
+	static const FString TAG_DEPRECATED_DEFAULTTEXT;
+	static const FString TAG_SOURCE;
+	static const FString TAG_SOURCE_TEXT;
+	static const FString TAG_PLATFORM_NAME;
+	static const FString TAG_METADATA;
+	static const FString TAG_METADATA_INFO;
+	static const FString TAG_METADATA_KEY;
+	static const FString NAMESPACE_DELIMITER;
 };
 
 #undef UE_API
diff --git a/Engine/Source/Developer/Localization/Public/Serialization/JsonInternationalizationMetadataSerializer.h b/Engine/Source/Developer/Localization/Public/Serialization/JsonInternationalizationMetadataSerializer.h
--- a/Engine/Source/Developer/Localization/Public/Serialization/JsonInternationalizationMetadataSerializer.h
+++ b/Engine/Source/Developer/Localization/Public/Serialization/JsonInternationalizationMetadataSerializer.h
@@ -2,14 +2,12 @@
 #pragma once
 
 #include "Containers/UnrealString.h"
-#include "CoreMinimal.h"
-#include "Dom/JsonObject.h"
 #include "Internationalization/InternationalizationMetadata.h"
+#include "JsonUtils/RapidJsonUtils.h"
 #include "Templates/SharedPointer.h"
 
 #define UE_API LOCALIZATION_API
 
-class FJsonObject;
 class FLocMetadataObject;
 
 class FJsonInternationalizationMetaDataSerializer
@@ -19,22 +17,22 @@ public:
 	 * Deserializes manifest metadata from a JSON object.
 	 *
 	 * @param InJsonObj The JSON object to serialize from.
-	 * @param OutMetadataObj The resulting metadata object or null if the serialize is not successful.
+	 * @param OutMetadataObj The resulting metadata object, or an unmodified object if the serialize is not successful.
 	 */
-	static UE_API void DeserializeMetadata( const TSharedRef< FJsonObject > JsonObj, TSharedPtr< FLocMetadataObject >& OutMetadataObj );
+	static UE_API void DeserializeMetadata( UE::Json::FConstObject JsonObj, TSharedPtr< FLocMetadataObject >& OutMetadataObj );
 
 	/**
 	 * Serializes manifest metadata object to a JSON object.
 	 *
 	 * @param Manifest The Internationalization metadata to serialize.
-	 * @param OutJsonObj The resulting JSON object or NULL if the serilize is not successful.
+	 * @param OutJsonObj The resulting JSON object or an unmodified object if the serilize is not successful.
 	 */
-	static UE_API void SerializeMetadata(  const TSharedRef< FLocMetadataObject > Metadata, TSharedPtr< FJsonObject >& OutJsonObj );
+	static UE_API void SerializeMetadata( const TSharedRef< FLocMetadataObject > Metadata, UE::Json::FValue& OutJsonObj, UE::Json::FAllocator& Allocator );
 
 	/**
 	 * Utility function that will convert metadata to string using the JSON metadata serializers.
 	 *
-	 * @param Manifest The Internationalization metadata to conver to string.
+	 * @param Manifest The Internationalization metadata to convert to string.
 	 */
 	static UE_API FString MetadataToString( const TSharedPtr<FLocMetadataObject> Metadata );
 };
diff --git a/Engine/Source/Developer/TranslationEditor/Private/TranslationDataManager.cpp b/Engine/Source/Developer/TranslationEditor/Private/TranslationDataManager.cpp
--- a/Engine/Source/Developer/TranslationEditor/Private/TranslationDataManager.cpp
+++ b/Engine/Source/Developer/TranslationEditor/Private/TranslationDataManager.cpp
@@ -280,17 +280,17 @@ bool FTranslationDataManager::WriteTranslationData(bool bForceWrite /*= false*/)
 
 		if (bForceWrite || bNeedsWrite)
 		{
-			TSharedRef<FJsonObject> FinalArchiveJsonObj = MakeShareable(new FJsonObject);
-			FJsonInternationalizationArchiveSerializer::SerializeArchive(Archive, FinalArchiveJsonObj);
+			FString FinalArchiveJson;
+			FJsonInternationalizationArchiveSerializer::SerializeArchiveToString(Archive, FinalArchiveJson);
 
-			bSuccess = WriteJSONToTextFile(FinalArchiveJsonObj, OpenedArchiveFilePath);
+			bSuccess = WriteJSONToTextFile(FinalArchiveJson, OpenedArchiveFilePath);
 		}
 	}
 
 	return bSuccess;
 }
 
-bool FTranslationDataManager::WriteJSONToTextFile(TSharedRef<FJsonObject>& Output, const FString& Filename)
+bool FTranslationDataManager::WriteJSONToTextFile(const FString& Output, const FString& Filename)
 {
 	bool CheckoutAndSaveWasSuccessful = true;
 	bool bPreviouslyCheckedOut = false;
@@ -326,11 +326,7 @@ bool FTranslationDataManager::WriteJSONToTextFile(TSharedRef<FJsonObject>& Outpu
 		if( CheckoutAndSaveWasSuccessful )
 		{
 			//Print the JSON data out to the ref file.
-			FString OutputString;
-			TSharedRef< TJsonWriter<> > Writer = TJsonWriterFactory<>::Create( &OutputString );
-			FJsonSerializer::Serialize( Output, Writer );
-
-			if (!FFileHelper::SaveStringToFile(OutputString, *Filename, FFileHelper::EEncodingOptions::ForceUnicode))
+			if (!FFileHelper::SaveStringToFile(Output, *Filename, FFileHelper::EEncodingOptions::ForceUnicode))
 			{
 				// If we already checked out the file, but cannot write it, perhaps the user checked it in via perforce, so try to check it out again
 				if (bPreviouslyCheckedOut)
diff --git a/Engine/Source/Developer/TranslationEditor/Private/TranslationDataManager.h b/Engine/Source/Developer/TranslationEditor/Private/TranslationDataManager.h
--- a/Engine/Source/Developer/TranslationEditor/Private/TranslationDataManager.h
+++ b/Engine/Source/Developer/TranslationEditor/Private/TranslationDataManager.h
@@ -96,8 +96,8 @@ private:
 	/** Retrieve an archive data structure from ArchiveFilePath */
 	TSharedPtr< FInternationalizationArchive > ReadArchive(const FString& ArchiveFilePath);
 
-	/** Write JSON file to text file */
-	bool WriteJSONToTextFile( TSharedRef<FJsonObject>& Output, const FString& Filename );
+	/** Write JSON string to text file */
+	bool WriteJSONToTextFile( const FString& Output, const FString& Filename );
 
 	/** Removes each UTranslationUnit in the passed array from the root set, allowing it to be garbage collected */
 	void RemoveTranslationUnitArrayfromRoot( TArray<UTranslationUnit*>& TranslationUnits );
