API comment clean-up

#rnx

diff --git a/Engine/Source/Developer/Localization/Public/Serialization/JsonInternationalizationArchiveSerializer.h b/Engine/Source/Developer/Localization/Public/Serialization/JsonInternationalizationArchiveSerializer.h
--- a/Engine/Source/Developer/Localization/Public/Serialization/JsonInternationalizationArchiveSerializer.h
+++ b/Engine/Source/Developer/Localization/Public/Serialization/JsonInternationalizationArchiveSerializer.h
@@ -27,7 +27,7 @@ public:
 	/**
 	 * Deserializes an archive from a JSON string.
 	 *
-	 * @param InJsonObj			The JSON string to serialize from.
+	 * @param InStr				The JSON string to deserialize from.
 	 * @param InArchive			The archive to populate from the JSON data.
 	 * @param InManifest		The manifest associated with the archive. May be null, but you won't be able to load archives with a version < FInternationalizationArchive::EFormatVersion::AddedKeys.
 	 * @param InNativeArchive	The native archive associated with the archive. May be null.
@@ -39,7 +39,7 @@ public:
 	/**
 	 * Deserializes an archive from a JSON file.
 	 *
-	 * @param InJsonFile		The path to the JSON file to serialize from.
+	 * @param InJsonFile		The path to the JSON file to deserialize from.
 	 * @param InArchive			The archive to populate from the JSON data.
 	 * @param InManifest		The manifest associated with the archive. May be null, but you won't be able to load archives with a version < FInternationalizationArchive::EFormatVersion::AddedKeys.
 	 * @param InNativeArchive	The native archive associated with the archive. May be null.
diff --git a/Engine/Source/Developer/Localization/Public/Serialization/JsonInternationalizationManifestSerializer.h b/Engine/Source/Developer/Localization/Public/Serialization/JsonInternationalizationManifestSerializer.h
--- a/Engine/Source/Developer/Localization/Public/Serialization/JsonInternationalizationManifestSerializer.h
+++ b/Engine/Source/Developer/Localization/Public/Serialization/JsonInternationalizationManifestSerializer.h
@@ -26,7 +26,7 @@ public:
 	/**
 	 * Deserializes a Internationalization manifest from a JSON string.
 	 *
-	 * @param InStr The JSON string to serialize from.
+	 * @param InStr The JSON string to deserialize from.
 	 * @param Manifest The populated Internationalization manifest.
 	 * @return true if deserialization was successful, false otherwise.
 	 */
@@ -35,7 +35,7 @@ public:
 	/**
 	 * Deserializes a Internationalization manifest from a JSON file.
 	 *
-	 * @param InJsonFile The path to the JSON file to serialize from.
+	 * @param InJsonFile The path to the JSON file to deserialize from.
 	 * @param Manifest The populated Internationalization manifest.
 	 * @return true if deserialization was successful, false otherwise.
 	 */
diff --git a/Engine/Source/Developer/Localization/Public/Serialization/JsonInternationalizationMetadataSerializer.h b/Engine/Source/Developer/Localization/Public/Serialization/JsonInternationalizationMetadataSerializer.h
--- a/Engine/Source/Developer/Localization/Public/Serialization/JsonInternationalizationMetadataSerializer.h
+++ b/Engine/Source/Developer/Localization/Public/Serialization/JsonInternationalizationMetadataSerializer.h
@@ -16,7 +16,7 @@ public:
 	/**
 	 * Deserializes manifest metadata from a JSON object.
 	 *
-	 * @param InJsonObj The JSON object to serialize from.
+	 * @param InJsonObj The JSON object to deserialize from.
 	 * @param OutMetadataObj The resulting metadata object, or an unmodified object if the serialize is not successful.
 	 */
 	static UE_API void DeserializeMetadata( UE::Json::FConstObject JsonObj, TSharedPtr< FLocMetadataObject >& OutMetadataObj );
@@ -24,15 +24,15 @@ public:
 	/**
 	 * Serializes manifest metadata object to a JSON object.
 	 *
-	 * @param Manifest The Internationalization metadata to serialize.
-	 * @param OutJsonObj The resulting JSON object or an unmodified object if the serilize is not successful.
+	 * @param Metadata The Internationalization metadata to serialize.
+	 * @param OutJsonObj The resulting JSON object or an unmodified object if the serialize is not successful.
 	 */
 	static UE_API void SerializeMetadata( const TSharedRef< FLocMetadataObject > Metadata, UE::Json::FValue& OutJsonObj, UE::Json::FAllocator& Allocator );
 
 	/**
-	 * Utility function that will convert metadata to string using the JSON metadata serializers.
+	 * Utility function that will convert metadata to string using the JSON metadata serializer.
 	 *
-	 * @param Manifest The Internationalization metadata to convert to string.
+	 * @param Metadata The Internationalization metadata to convert to string.
 	 */
 	static UE_API FString MetadataToString( const TSharedPtr<FLocMetadataObject> Metadata );
 };
