API comment clean-up

#rnx

diff --git a/Engine/Source/Developer/Localization/Public/Serialization/JsonInternationalizationMetadataSerializer.h b/Engine/Source/Developer/Localization/Public/Serialization/JsonInternationalizationMetadataSerializer.h
--- a/Engine/Source/Developer/Localization/Public/Serialization/JsonInternationalizationMetadataSerializer.h
+++ b/Engine/Source/Developer/Localization/Public/Serialization/JsonInternationalizationMetadataSerializer.h
@@ -16,7 +16,7 @@ public:
 	/**
 	 * Deserializes manifest metadata from a JSON object.
 	 *
-	 * @param InJsonObj The JSON object to deserialize from.
+	 * @param JsonObj The JSON object to deserialize from.
 	 * @param OutMetadataObj The resulting metadata object, or an unmodified object if the serialize is not successful.
 	 */
 	static UE_API void DeserializeMetadata( UE::Json::FConstObject JsonObj, TSharedPtr< FLocMetadataObject >& OutMetadataObj );
