Fixed logic error in FLocalizationSCC::CleanUp

#rnx

diff --git a/Engine/Source/Developer/Localization/Private/LocalizationSourceControlUtil.cpp b/Engine/Source/Developer/Localization/Private/LocalizationSourceControlUtil.cpp
--- a/Engine/Source/Developer/Localization/Private/LocalizationSourceControlUtil.cpp
+++ b/Engine/Source/Developer/Localization/Private/LocalizationSourceControlUtil.cpp
@@ -271,7 +271,6 @@ bool FLocalizationSCC::CleanUp(FText& OutError)
 		{
 			check(CheckedOutFilesArray.Num() == SourceControlStates.Num());
 
-			TArray<FString> FilesFailedToRevert;
 			for (int32 Index = 0; Index < CheckedOutFilesArray.Num(); ++Index)
 			{
 				FString& CheckedOutFile = CheckedOutFilesArray[Index];
@@ -281,13 +280,9 @@ bool FLocalizationSCC::CleanUp(FText& OutError)
 				{
 					CheckedOutFiles.Add(MoveTemp(CheckedOutFile));
 				}
-				else
-				{
-					FilesFailedToRevert.Add(MoveTemp(CheckedOutFile));
-				}
 			}
 
-			OutError = FText::Format(LOCTEXT("CleanUpFailed.WithFiles", "FLocalizationSCC::CleanUp failed to revert: {0}"), FText::FromString(FString::Join(FilesFailedToRevert, TEXT(", "))));
+			OutError = FText::Format(LOCTEXT("CleanUpFailed.WithFiles", "FLocalizationSCC::CleanUp failed to revert: {0}"), FText::FromString(FString::Join(CheckedOutFiles, TEXT(", "))));
 		}
 		else
 		{
