Sort the manifest used to build the chunked LocRes as if it had just been loaded

This keeps the LocRes data on disk consistent with previous versions of UE, and avoids a large patch from the LocRes data being internally reordered

#rnx

diff --git a/Engine/Source/Developer/Localization/Private/LocTextHelper.cpp b/Engine/Source/Developer/Localization/Private/LocTextHelper.cpp
--- a/Engine/Source/Developer/Localization/Private/LocTextHelper.cpp
+++ b/Engine/Source/Developer/Localization/Private/LocTextHelper.cpp
@@ -677,6 +677,12 @@ void FLocTextHelper::TrimManifest()
 	}
 }
 
+void FLocTextHelper::SortManifest()
+{
+	checkf(Manifest.IsValid(), TEXT("Attempted to enumerate source texts, but no manifest has been loaded!"));
+	FJsonInternationalizationManifestSerializer::SortManifest(Manifest.ToSharedRef());
+}
+
 bool FLocTextHelper::HasNativeArchive() const
 {
 	return HasArchive(NativeCulture);
@@ -1063,6 +1069,8 @@ void FLocTextHelper::EnumerateSourceTexts(const FEnumerateSourceTextsFuncPtr& In
 
 void FLocTextHelper::EnumerateSourceTextsParallel(const FEnumerateSourceTextsFuncPtr& InCallback, const bool InCheckDependencies) const
 {
+	checkf(Manifest.IsValid(), TEXT("Attempted to enumerate source texts, but no manifest has been loaded!"));
+
 	std::atomic<bool> bBreak = false;
 
 	TArray<TPair<FString, TSharedRef<FManifestEntry>>> ManifestEntries = Manifest->GetEntriesBySourceTextAsArray();
diff --git a/Engine/Source/Developer/Localization/Private/Serialization/JsonInternationalizationManifestSerializer.cpp b/Engine/Source/Developer/Localization/Private/Serialization/JsonInternationalizationManifestSerializer.cpp
--- a/Engine/Source/Developer/Localization/Private/Serialization/JsonInternationalizationManifestSerializer.cpp
+++ b/Engine/Source/Developer/Localization/Private/Serialization/JsonInternationalizationManifestSerializer.cpp
@@ -115,6 +115,20 @@ bool FJsonInternationalizationManifestSerializer::SerializeManifestToFile( TShar
 }
 
 
+void FJsonInternationalizationManifestSerializer::SortManifest(const TSharedRef<FInternationalizationManifest>& Manifest)
+{
+	TSharedRef<FStructuredEntry> RootElement = MakeShared<FStructuredEntry>(FString());
+
+	// Convert the manifest into the structured data we use for JSON serialization
+	GenerateStructuredData(Manifest, RootElement);
+	SortStructuredData(RootElement);
+
+	// Convert the structured data back into the in-memory manifest
+	Manifest->ClearEntries();
+	StructuredDataToManifest(RootElement, Manifest);
+}
+
+
 bool FJsonInternationalizationManifestSerializer::DeserializeInternal( UE::Json::FConstObject InJsonObj, TSharedRef< FInternationalizationManifest > Manifest, const FName PlatformName )
 {
 	if (TOptional<int32> FormatVersion = UE::Json::GetInt32Field(InJsonObj, *TAG_FORMATVERSION))
@@ -408,6 +422,28 @@ void FJsonInternationalizationManifestSerializer::SortStructuredData( TSharedPtr
 }
 
 
+void FJsonInternationalizationManifestSerializer::StructuredDataToManifest(const TSharedPtr<const FStructuredEntry>& InElement, const TSharedRef<FInternationalizationManifest>& Manifest)
+{
+	if (!InElement)
+	{
+		return;
+	}
+
+	for (const TSharedPtr<FManifestEntry>& Entry : InElement->ManifestEntries)
+	{
+		for (const FManifestContext& Context : Entry->Contexts)
+		{
+			Manifest->AddSource(Entry->Namespace, Entry->Source, Context);
+		}
+	}
+
+	for (const TTuple<FString, TSharedPtr<FStructuredEntry>>& SubElementPair : InElement->SubNamespaces)
+	{
+		StructuredDataToManifest(SubElementPair.Value, Manifest);
+	}
+}
+
+
 void FJsonInternationalizationManifestSerializer::StructuredDataToJsonObj( TSharedPtr< const FStructuredEntry > InElement, UE::Json::FObject JsonObj, UE::Json::FAllocator& Allocator )
 {
 	JsonObj.AddMember(UE::Json::MakeStringRef(TAG_NAMESPACE), UE::Json::MakeStringValue(InElement->Namespace, Allocator), Allocator);
diff --git a/Engine/Source/Developer/Localization/Public/LocTextHelper.h b/Engine/Source/Developer/Localization/Public/LocTextHelper.h
--- a/Engine/Source/Developer/Localization/Public/LocTextHelper.h
+++ b/Engine/Source/Developer/Localization/Public/LocTextHelper.h
@@ -396,6 +396,11 @@ public:
 	 */
 	UE_API void TrimManifest();
 
+	/**
+	 * Sort the currently loaded manifest into a deterministic order (as if it had been serialized and then deserialized from JSON).
+	 */
+	UE_API void SortManifest();
+
 	/**
 	 * Check to see whether we've loaded the native archive.
 	 */
diff --git a/Engine/Source/Developer/Localization/Public/Serialization/JsonInternationalizationManifestSerializer.h b/Engine/Source/Developer/Localization/Public/Serialization/JsonInternationalizationManifestSerializer.h
--- a/Engine/Source/Developer/Localization/Public/Serialization/JsonInternationalizationManifestSerializer.h
+++ b/Engine/Source/Developer/Localization/Public/Serialization/JsonInternationalizationManifestSerializer.h
@@ -59,6 +59,11 @@ public:
 	 */
 	static UE_API bool SerializeManifestToFile( TSharedRef< const FInternationalizationManifest > Manifest, const FString& InJsonFile );
 
+	/**
+	 * Sort an Internationalization manifest as if it had been serialized and then deserialized from JSON.
+	 */
+	static UE_API void SortManifest(const TSharedRef<FInternationalizationManifest>& Manifest);
+
 private:
 
 	struct FStructuredEntry
@@ -119,6 +124,14 @@ private:
 	 */
 	static void SortStructuredData( TSharedPtr< FStructuredEntry > InElement );
 
+	/**
+	 * Populates an Internationalization manifest from data that has been structured based on namespace.
+	 *
+	 * @param InElement The structured data.
+	 * @param Manifest The Internationalization manifest that will store the data.
+	 */
+	static void StructuredDataToManifest(const TSharedPtr<const FStructuredEntry>& InElement, const TSharedRef<FInternationalizationManifest>& Manifest);
+
 	/**
 	 * Populates a JSON object from Internationalization manifest data that has been structured based on namespace.
 	 *
diff --git a/Engine/Source/Editor/UnrealEd/Private/LocalizationChunkDataGenerator.cpp b/Engine/Source/Editor/UnrealEd/Private/LocalizationChunkDataGenerator.cpp
--- a/Engine/Source/Editor/UnrealEd/Private/LocalizationChunkDataGenerator.cpp
+++ b/Engine/Source/Editor/UnrealEd/Private/LocalizationChunkDataGenerator.cpp
@@ -186,6 +186,9 @@ void FLocalizationChunkDataGenerator::GenerateChunkDataFiles(const int32 InChunk
 		// This won't add them to the build, and we don't care if this fails as it's only for debugging
 		ChunkLocTextHelper.SaveAll();
 
+		// Sort the manifest to ensure the LocRes output is deterministic, as this was populated via a ParallelFor
+		ChunkLocTextHelper.SortManifest();
+
 		// Produce the LocMeta file for the chunk target
 		{
 			const FString ChunkLocMetaFilename = ChunkTargetContentRoot / FString::Printf(TEXT("%s.locmeta"), *ChunkTargetName);
diff --git a/Engine/Source/Runtime/Core/Private/Internationalization/InternationalizationManifest.cpp b/Engine/Source/Runtime/Core/Private/Internationalization/InternationalizationManifest.cpp
--- a/Engine/Source/Runtime/Core/Private/Internationalization/InternationalizationManifest.cpp
+++ b/Engine/Source/Runtime/Core/Private/Internationalization/InternationalizationManifest.cpp
@@ -289,6 +289,12 @@ TSharedPtr<FManifestEntry> FInternationalizationManifest::FindEntryByKey(const F
 	return nullptr;
 }
 
+void FInternationalizationManifest::ClearEntries()
+{
+	EntriesBySourceText.Empty();
+	EntriesByKey.Empty();
+}
+
 const FManifestContext* FManifestEntry::FindContext(const FLocKey& ContextKey, const TSharedPtr<FLocMetadataObject>& KeyMetadata) const
 {
 	for (const FManifestContext& Context : Contexts)
diff --git a/Engine/Source/Runtime/Core/Private/Internationalization/TextLocalizationResource.cpp b/Engine/Source/Runtime/Core/Private/Internationalization/TextLocalizationResource.cpp
--- a/Engine/Source/Runtime/Core/Private/Internationalization/TextLocalizationResource.cpp
+++ b/Engine/Source/Runtime/Core/Private/Internationalization/TextLocalizationResource.cpp
@@ -480,22 +480,6 @@ bool FTextLocalizationResource::SaveToArchive(FArchive& Archive, const FTextKey&
 		KeysTable.Emplace(EntryPair.Key.GetKey(), &EntryPair.Value);
 	}
 
-	// Sort the namespace keys
-	auto SortTextKeysPredicate = [](const FTextKey& A, const FTextKey& B) -> bool
-		{
-			TStringBuilder<256> AString, BString;
-			A.ToString(AString);
-			B.ToString(BString);
-			
-			return AString.ToView() < BString.ToView();
-		};
-
-	for (auto& NamespacePair : Namespaces)
-	{
-		NamespacePair.Value.KeySort(SortTextKeysPredicate);
-	}
-	Namespaces.KeySort(SortTextKeysPredicate);
-
 	// Write entries count
 	uint32 EntriesCount = Entries.Num();
 	Archive << EntriesCount;
diff --git a/Engine/Source/Runtime/Core/Public/Internationalization/InternationalizationManifest.h b/Engine/Source/Runtime/Core/Public/Internationalization/InternationalizationManifest.h
--- a/Engine/Source/Runtime/Core/Public/Internationalization/InternationalizationManifest.h
+++ b/Engine/Source/Runtime/Core/Public/Internationalization/InternationalizationManifest.h
@@ -148,6 +148,8 @@ public:
 
 	CORE_API TSharedPtr<FManifestEntry> FindEntryByKey(const FLocKey& Namespace, const FLocKey& Key, const FString* SourceText = nullptr) const;
 
+	CORE_API void ClearEntries();
+
 	FManifestEntryByLocKeyContainer::TConstIterator GetEntriesByKeyIterator() const
 	{
 		return EntriesByKey.CreateConstIterator();
