Parallelized the internal PO import/export and LocRes compile steps

This is in addition to the step-level parallelization, as we can safely read/write a single archive at a time as they're pre-loaded before we start the ParallelFor

diff --git a/Engine/Source/Developer/Localization/Private/LocTextHelper.cpp b/Engine/Source/Developer/Localization/Private/LocTextHelper.cpp
--- a/Engine/Source/Developer/Localization/Private/LocTextHelper.cpp
+++ b/Engine/Source/Developer/Localization/Private/LocTextHelper.cpp
@@ -860,8 +860,7 @@ void FLocTextHelper::TrimArchive(const FString& InCulture)
 		return true; // continue enumeration
 	}, true);
 
-	Archives.Remove(InCulture);
-	Archives.Add(InCulture, TrimmedArchive);
+	Archives[InCulture] = TrimmedArchive;
 }
 
 bool FLocTextHelper::LoadAll(const ELocTextHelperLoadFlags InLoadFlags, FText* OutError)
diff --git a/Engine/Source/Developer/Localization/Private/LocalizationSourceControlUtil.cpp b/Engine/Source/Developer/Localization/Private/LocalizationSourceControlUtil.cpp
--- a/Engine/Source/Developer/Localization/Private/LocalizationSourceControlUtil.cpp
+++ b/Engine/Source/Developer/Localization/Private/LocalizationSourceControlUtil.cpp
@@ -6,11 +6,16 @@
 #include "HAL/PlatformFileManager.h"
 #include "ISourceControlModule.h"
 #include "ISourceControlProvider.h"
+#include "Logging/StructuredLog.h"
 #include "Misc/Paths.h"
 #include "SourceControlHelpers.h"
 #include "SourceControlOperations.h"
 
 DEFINE_LOG_CATEGORY_STATIC(LogLocalizationSourceControl, Log, All);
+namespace LocalizationSourceControlUtil
+{
+	static constexpr int32 LocalizationLogIdentifier = 304;
+}
 
 #define LOCTEXT_NAMESPACE "LocalizationSourceControl"
 
@@ -35,18 +40,17 @@ FLocalizationSCC::~FLocalizationSCC()
 
 void FLocalizationSCC::BeginParallelTasks()
 {
-	checkf(IsInGameThread(), TEXT("FLocalizationSCC::BeginParallelTasks must be called on the game-thread"));
-	checkf(!bRunningParallelTasks, TEXT("FLocalizationSCC::BeginParallelTasks was called while parallel tasks are already running"));
-
-	bRunningParallelTasks = true;
+	++ParallelTasksCount;
 }
 
 bool FLocalizationSCC::EndParallelTasks(FText& OutError)
 {
-	checkf(IsInGameThread(), TEXT("FLocalizationSCC::EndParallelTasks must be called on the game-thread"));
-	checkf(bRunningParallelTasks, TEXT("FLocalizationSCC::EndParallelTasks was called while no parallel tasks are running"));
+	if (--ParallelTasksCount > 0)
+	{
+		return true;
+	}
 
-	bRunningParallelTasks = false;
+	checkf(IsInGameThread(), TEXT("FLocalizationSCC::EndParallelTasks must be called on the game-thread when closing the final block"));
 
 	if (DeferredCheckedOutFiles.IsEmpty())
 	{
@@ -126,7 +130,7 @@ bool FLocalizationSCC::CheckOutFile(const FString& InFile, FText& OutError)
 
 	FString AbsoluteFilename = FPaths::ConvertRelativePathToFull(InFile);
 
-	if (bRunningParallelTasks)
+	if (ParallelTasksCount > 0)
 	{
 		UE::TScopeLock _(FilesMutex);
 
@@ -188,7 +192,7 @@ bool FLocalizationSCC::CheckOutFile(const FString& InFile, FText& OutError)
 bool FLocalizationSCC::CheckinFiles(const FText& InChangeDescription, FText& OutError)
 {
 	checkf(IsInGameThread(), TEXT("FLocalizationSCC::CheckinFiles must be called on the game-thread"));
-	checkf(!bRunningParallelTasks, TEXT("FLocalizationSCC::CheckinFiles was called while parallel tasks are running"));
+	checkf(ParallelTasksCount == 0, TEXT("FLocalizationSCC::CheckinFiles was called while parallel tasks are running"));
 
 	if (CheckedOutFiles.IsEmpty())
 	{
@@ -245,7 +249,7 @@ bool FLocalizationSCC::CheckinFiles(const FText& InChangeDescription, FText& Out
 bool FLocalizationSCC::CleanUp(FText& OutError)
 {
 	checkf(IsInGameThread(), TEXT("FLocalizationSCC::CleanUp must be called on the game-thread"));
-	checkf(!bRunningParallelTasks, TEXT("FLocalizationSCC::CleanUp was called while parallel tasks are running"));
+	checkf(ParallelTasksCount == 0, TEXT("FLocalizationSCC::CleanUp was called while parallel tasks are running"));
 
 	if (CheckedOutFiles.IsEmpty())
 	{
@@ -297,7 +301,7 @@ bool FLocalizationSCC::CleanUp(FText& OutError)
 bool FLocalizationSCC::IsReady(FText& OutError) const
 {
 	checkf(IsInGameThread(), TEXT("FLocalizationSCC::IsReady must be called on the game-thread"));
-	checkf(!bRunningParallelTasks, TEXT("FLocalizationSCC::IsReady was called while parallel tasks are running"));
+	checkf(ParallelTasksCount == 0, TEXT("FLocalizationSCC::IsReady was called while parallel tasks are running"));
 
 	if (!ISourceControlModule::Get().IsEnabled())
 	{
@@ -317,7 +321,7 @@ bool FLocalizationSCC::IsReady(FText& OutError) const
 bool FLocalizationSCC::RevertFile(const FString& InFile, FText& OutError)
 {
 	checkf(IsInGameThread(), TEXT("FLocalizationSCC::RevertFile must be called on the game-thread"));
-	checkf(!bRunningParallelTasks, TEXT("FLocalizationSCC::RevertFile was called while parallel tasks are running"));
+	checkf(ParallelTasksCount == 0, TEXT("FLocalizationSCC::RevertFile was called while parallel tasks are running"));
 
 	if (InFile.IsEmpty() || InFile.StartsWith(TEXT("\\\\")))
 	{
@@ -343,28 +347,59 @@ bool FLocalizationSCC::RevertFile(const FString& InFile, FText& OutError)
 	return true;
 }
 
+void FLocFileSCCNotifies::BeginParallelTasks()
+{
+	if (SourceControlInfo)
+	{
+		SourceControlInfo->BeginParallelTasks();
+	}
+}
+
+void FLocFileSCCNotifies::EndParallelTasks()
+{
+	if (SourceControlInfo)
+	{
+		FText ErrorMsg;
+		if (!SourceControlInfo->EndParallelTasks(ErrorMsg))
+		{
+			UE_LOGFMT(LogLocalizationSourceControl, Error, "{error}",
+				("error", ErrorMsg.ToString()),
+				("id", LocalizationSourceControlUtil::LocalizationLogIdentifier)
+			);
+		}
+	}
+}
+
 void FLocFileSCCNotifies::PreFileWrite(const FString& InFilename)
 {
-	if (SourceControlInfo.IsValid() && FPaths::FileExists(InFilename))
+	if (SourceControlInfo && FPaths::FileExists(InFilename))
 	{
 		// File already exists, so check it out before writing to it
 		FText ErrorMsg;
 		if (!SourceControlInfo->CheckOutFile(InFilename, ErrorMsg))
 		{
-			UE_LOG(LogLocalizationSourceControl, Error, TEXT("Failed to check out file '%s'. %s"), *InFilename, *ErrorMsg.ToString());
+			UE_LOGFMT(LogLocalizationSourceControl, Error, "Failed to check out file '{filename}'. {error}",
+				("filename", InFilename),
+				("error", ErrorMsg.ToString()),
+				("id", LocalizationSourceControlUtil::LocalizationLogIdentifier)
+			);
 		}
 	}
 }
 
 void FLocFileSCCNotifies::PostFileWrite(const FString& InFilename)
 {
-	if (SourceControlInfo.IsValid())
+	if (SourceControlInfo)
 	{
 		// If the file didn't exist before then this will add it, otherwise it will do nothing
 		FText ErrorMsg;
 		if (!SourceControlInfo->CheckOutFile(InFilename, ErrorMsg))
 		{
-			UE_LOG(LogLocalizationSourceControl, Error, TEXT("Failed to check out file '%s'. %s"), *InFilename, *ErrorMsg.ToString());
+			UE_LOGFMT(LogLocalizationSourceControl, Error, "Failed to check out file '{filename}'. {error}",
+				("filename", InFilename),
+				("error", ErrorMsg.ToString()),
+				("id", LocalizationSourceControlUtil::LocalizationLogIdentifier)
+			);
 		}
 	}
 }
diff --git a/Engine/Source/Developer/Localization/Private/PortableObjectPipeline.cpp b/Engine/Source/Developer/Localization/Private/PortableObjectPipeline.cpp
--- a/Engine/Source/Developer/Localization/Private/PortableObjectPipeline.cpp
+++ b/Engine/Source/Developer/Localization/Private/PortableObjectPipeline.cpp
@@ -1,6 +1,8 @@
 // Copyright Epic Games, Inc. All Rights Reserved.
 
 #include "PortableObjectPipeline.h"
+
+#include "Async/ParallelFor.h"
 #include "Misc/FileHelper.h"
 #include "Misc/Paths.h"
 #include "Internationalization/InternationalizationMetadata.h"
@@ -470,7 +472,8 @@ bool PortableObjectPipeline::ImportAll(FLocTextHelper& InLocTextHelper, const FS
 {
 	// We may only have a single culture if using this setting
 	const bool bSingleCultureMode = !bUseCultureDirectory;
-	if (bSingleCultureMode && (InLocTextHelper.GetAllCultures(bSingleCultureMode).Num() != 1 || InLocTextHelper.ShouldSplitPlatformData()))
+	const TArray<FString> Cultures = InLocTextHelper.GetAllCultures(bSingleCultureMode);
+	if (bSingleCultureMode && (Cultures.Num() != 1 || InLocTextHelper.ShouldSplitPlatformData()))
 	{
 		UE_LOG(LogPortableObjectPipeline, Error, TEXT("bUseCultureDirectory may only be used with a single culture when not splitting platform data."));
 		return false;
@@ -482,10 +485,19 @@ bool PortableObjectPipeline::ImportAll(FLocTextHelper& InLocTextHelper, const FS
 	TMap<FName, TSharedRef<FInternationalizationManifest>> PerPlatformManifests;
 	BuildCollapsedManifest(InLocTextHelper, InTextCollapseMode, CollapsedData, PlatformAgnosticManifest, PerPlatformManifests);
 
+	TSharedPtr<ILocFileNotifies> LocFileNotifies = InLocTextHelper.GetLocFileNotifies();
+
+	if (LocFileNotifies)
+	{
+		LocFileNotifies->BeginParallelTasks();
+	}
+
 	// Process the desired cultures
-	bool bSuccess = true;
-	for (const FString& CultureName : InLocTextHelper.GetAllCultures(bSingleCultureMode))
+	std::atomic<bool> bSuccess = true;
+	ParallelFor(TEXT("PortableObjectPipeline::ImportAll"), Cultures.Num(), 1, [&](int32 Index)
 	{
+		const FString& CultureName = Cultures[Index];
+
 		auto ImportSinglePortableObject = [&](const FName InPlatformName) -> bool
 		{
 			// Which path should we use for the PO?
@@ -510,11 +522,22 @@ bool PortableObjectPipeline::ImportAll(FLocTextHelper& InLocTextHelper, const FS
 			return ImportPortableObject(InLocTextHelper, CultureName, POFilePath, InTextCollapseMode, InPOFormat, CollapsedData);
 		};
 
-		bSuccess &= ImportSinglePortableObject(FName());
+		if (!ImportSinglePortableObject(FName()))
+		{
+			bSuccess = false;
+		}
 		for (const auto& PerPlatformManifestPair : PerPlatformManifests)
 		{
-			bSuccess &= ImportSinglePortableObject(PerPlatformManifestPair.Key);
+			if (!ImportSinglePortableObject(PerPlatformManifestPair.Key))
+			{
+				bSuccess = false;
+			}
 		}
+	});
+
+	if (LocFileNotifies)
+	{
+		LocFileNotifies->EndParallelTasks();
 	}
 
 	return bSuccess;
@@ -542,7 +565,8 @@ bool PortableObjectPipeline::ExportAll(FLocTextHelper& InLocTextHelper, const FS
 {
 	// We may only have a single culture if using this setting
 	const bool bSingleCultureMode = !bUseCultureDirectory;
-	if (bSingleCultureMode && (InLocTextHelper.GetAllCultures(bSingleCultureMode).Num() != 1 || InLocTextHelper.ShouldSplitPlatformData()))
+	const TArray<FString> Cultures = InLocTextHelper.GetAllCultures(bSingleCultureMode);
+	if (bSingleCultureMode && (Cultures.Num() != 1 || InLocTextHelper.ShouldSplitPlatformData()))
 	{
 		UE_LOG(LogPortableObjectPipeline, Error, TEXT("bUseCultureDirectory may only be used with a single culture when not splitting platform data."));
 		return false;
@@ -561,10 +585,19 @@ bool PortableObjectPipeline::ExportAll(FLocTextHelper& InLocTextHelper, const FS
 	TMap<FName, TSharedRef<FInternationalizationManifest>> PerPlatformManifests;
 	BuildCollapsedManifest(InLocTextHelper, InTextCollapseMode, CollapsedData, PlatformAgnosticManifest, PerPlatformManifests);
 
+	TSharedPtr<ILocFileNotifies> LocFileNotifies = InLocTextHelper.GetLocFileNotifies();
+
+	if (LocFileNotifies)
+	{
+		LocFileNotifies->BeginParallelTasks();
+	}
+
 	// Process the desired cultures
-	bool bSuccess = true;
-	for (const FString& CultureName : InLocTextHelper.GetAllCultures(bSingleCultureMode))
+	std::atomic<bool> bSuccess = true;
+	ParallelFor(TEXT("PortableObjectPipeline::ExportAll"), Cultures.Num(), 1, [&](int32 Index)
 	{
+		const FString& CultureName = Cultures[Index];
+
 		auto ExportSinglePortableObject = [&](const TSharedRef<FInternationalizationManifest>& InCollapsedManifest, const FName InPlatformName) -> bool
 		{
 			// Which path should we use for the PO?
@@ -589,11 +622,22 @@ bool PortableObjectPipeline::ExportAll(FLocTextHelper& InLocTextHelper, const FS
 			return ExportPortableObject(InLocTextHelper, CultureName, POFilePath, InTextCollapseMode, InPOFormat, InCollapsedManifest, CollapsedData, bShouldPersistComments);
 		};
 
-		bSuccess &= ExportSinglePortableObject(PlatformAgnosticManifest.ToSharedRef(), FName());
+		if (!ExportSinglePortableObject(PlatformAgnosticManifest.ToSharedRef(), FName()))
+		{
+			bSuccess = false;
+		}
 		for (const auto& PerPlatformManifestPair : PerPlatformManifests)
 		{
-			bSuccess &= ExportSinglePortableObject(PerPlatformManifestPair.Value, PerPlatformManifestPair.Key);
+			if (!ExportSinglePortableObject(PerPlatformManifestPair.Value, PerPlatformManifestPair.Key))
+			{
+				bSuccess = false;
+			}
 		}
+	});
+
+	if (LocFileNotifies)
+	{
+		LocFileNotifies->EndParallelTasks();
 	}
 
 	return bSuccess;
diff --git a/Engine/Source/Developer/Localization/Public/LocTextHelper.h b/Engine/Source/Developer/Localization/Public/LocTextHelper.h
--- a/Engine/Source/Developer/Localization/Public/LocTextHelper.h
+++ b/Engine/Source/Developer/Localization/Public/LocTextHelper.h
@@ -68,7 +68,13 @@ class ILocFileNotifies
 {
 public:
 	/** Virtual destructor */
-	virtual ~ILocFileNotifies() {}
+	virtual ~ILocFileNotifies() = default;
+
+	/** Start a block of parallel tasks */
+	virtual void BeginParallelTasks() = 0;
+
+	/** Stop a block of parallel tasks */
+	virtual void EndParallelTasks() = 0;
 
 	/** Called prior to reading the given file on disk */
 	virtual void PreFileRead(const FString& InFilename) = 0;
diff --git a/Engine/Source/Developer/Localization/Public/LocalizationSourceControlUtil.h b/Engine/Source/Developer/Localization/Public/LocalizationSourceControlUtil.h
--- a/Engine/Source/Developer/Localization/Public/LocalizationSourceControlUtil.h
+++ b/Engine/Source/Developer/Localization/Public/LocalizationSourceControlUtil.h
@@ -19,7 +19,7 @@ public:
 	UE_API FLocalizationSCC();
 	UE_API ~FLocalizationSCC();
 
-	/** Start a block of parallel tasks; this will defer check-out requests until EndParallelTasks is called */
+	/** Start a block of parallel tasks (may be nested); this will defer check-out requests until the final EndParallelTasks is called */
 	UE_API void BeginParallelTasks();
 	/** Stop a block of parallel tasks; this will attempt to perform any check-outs deferred during the block */
 	UE_API bool EndParallelTasks(FText& OutError);
@@ -40,10 +40,10 @@ private:
 	TSet<FString> CheckedOutFiles;
 	/** Set of files that will be checked-out from SCC via CheckOutFile, but had to be deferred due to running in parallel mode */
 	TSet<FString> DeferredCheckedOutFiles;
-	/** Mutex protecting CheckedOutFiles and DeferredCheckedOutFiles when bRunningParallelTasks is true */
+	/** Mutex protecting CheckedOutFiles and DeferredCheckedOutFiles when ParallelTasksCount > 0 */
 	mutable UE::FMutex FilesMutex;
-	/** True if this we're being used from a parallel tasks block, and need to defer check-out requests */
-	std::atomic<bool> bRunningParallelTasks = false;
+	/** >0 if this we're being used from a parallel tasks block, and need to defer check-out requests */
+	std::atomic<uint8> ParallelTasksCount = 0;
 };
 
 class FLocFileSCCNotifies : public ILocFileNotifies
@@ -55,9 +55,11 @@ public:
 	}
 
 	/** Virtual destructor */
-	virtual ~FLocFileSCCNotifies() {}
+	virtual ~FLocFileSCCNotifies() = default;
 
 	//~ ILocFileNotifies interface
+	UE_API virtual void BeginParallelTasks() override;
+	UE_API virtual void EndParallelTasks() override;
 	virtual void PreFileRead(const FString& InFilename) override {}
 	virtual void PostFileRead(const FString& InFilename) override {}
 	UE_API virtual void PreFileWrite(const FString& InFilename) override;
diff --git a/Engine/Source/Editor/UnrealEd/Private/Commandlets/GenerateTextLocalizationResourceCommandlet.cpp b/Engine/Source/Editor/UnrealEd/Private/Commandlets/GenerateTextLocalizationResourceCommandlet.cpp
--- a/Engine/Source/Editor/UnrealEd/Private/Commandlets/GenerateTextLocalizationResourceCommandlet.cpp
+++ b/Engine/Source/Editor/UnrealEd/Private/Commandlets/GenerateTextLocalizationResourceCommandlet.cpp
@@ -2,6 +2,7 @@
 
 #include "Commandlets/GenerateTextLocalizationResourceCommandlet.h"
 
+#include "Async/ParallelFor.h"
 #include "Commandlets/Commandlet.h"
 #include "Containers/Array.h"
 #include "Containers/Map.h"
@@ -186,16 +187,30 @@ int32 UGenerateTextLocalizationResourceCommandlet::Main(const FString& Params)
 			return FTextLocalizationResourceGenerator::GenerateLocMeta(LocTextHelper, ResourceName, LocMeta) && LocMeta.SaveToFile(InSaveFileName);
 		});
 
-		if (!bLocMetaFileSaved)
+		if (bLocMetaFileSaved)
+		{
+			UE_LOG(LogGenerateTextLocalizationResourceCommandlet, Log, TEXT("Generated file %s"), *TextLocalizationMetaDataResourcePath);
+		}
+		else
 		{
 			UE_LOG(LogGenerateTextLocalizationResourceCommandlet, Error, TEXT("Could not write file %s"), *TextLocalizationMetaDataResourcePath);
 			return -1;
 		}
 	}
 
+	TSharedPtr<ILocFileNotifies> LocFileNotifies = LocTextHelper.GetLocFileNotifies();
+
+	if (LocFileNotifies)
+	{
+		LocFileNotifies->BeginParallelTasks();
+	}
+
 	// Generate the LocRes file for each culture
-	for (const FString& CultureName : CulturesToGenerate)
+	std::atomic<bool> bSuccess = true;
+	ParallelFor(TEXT("GenerateTextLocalizationResourceCommandlet"), CulturesToGenerate.Num(), 1, [&](int32 Index)
 	{
+		const FString& CultureName = CulturesToGenerate[Index];
+
 		auto GenerateSingleLocRes = [this, &DestinationPath, &CultureName, &ResourceName](const FTextLocalizationResource& InLocRes, const FName InPlatformName) -> bool
 		{
 			FString TextLocalizationResourcePath;
@@ -213,7 +228,11 @@ int32 UGenerateTextLocalizationResourceCommandlet::Main(const FString& Params)
 				return InLocRes.SaveToFile(InSaveFileName);
 			});
 
-			if (!bLocResFileSaved)
+			if (bLocResFileSaved)
+			{
+				UE_LOG(LogGenerateTextLocalizationResourceCommandlet, Log, TEXT("Generated file %s"), *TextLocalizationResourcePath);
+			}
+			else
 			{
 				UE_LOG(LogGenerateTextLocalizationResourceCommandlet, Error, TEXT("Could not write file %s"), *TextLocalizationResourcePath);
 				return false;
@@ -228,18 +247,31 @@ int32 UGenerateTextLocalizationResourceCommandlet::Main(const FString& Params)
 		if (!FTextLocalizationResourceGenerator::GenerateLocRes(LocTextHelper, CultureName, GenerateFlags, LocResId, PlatformAgnosticLocRes, PerPlatformLocRes))
 		{
 			UE_LOG(LogGenerateTextLocalizationResourceCommandlet, Error, TEXT("Failed to generate LocRes %s"), *LocResId.ToString());
-			return false;
+			bSuccess = false;
+			return;
 		}
 	
-		bool bSuccess = GenerateSingleLocRes(PlatformAgnosticLocRes, FName());
-		for (const auto& PerPlatformLocResPair : PerPlatformLocRes)
+		if (!GenerateSingleLocRes(PlatformAgnosticLocRes, FName()))
 		{
-			bSuccess &= GenerateSingleLocRes(*PerPlatformLocResPair.Value, PerPlatformLocResPair.Key);
+			bSuccess = false;
 		}
-		if (!bSuccess)
+		for (const auto& PerPlatformLocResPair : PerPlatformLocRes)
 		{
-			return -1;
+			if (!GenerateSingleLocRes(*PerPlatformLocResPair.Value, PerPlatformLocResPair.Key))
+			{
+				bSuccess = false;
+			}
 		}
+	});
+
+	if (LocFileNotifies)
+	{
+		LocFileNotifies->EndParallelTasks();
+	}
+
+	if (!bSuccess)
+	{
+		return -1;
 	}
 
 	return 0;
