Re-added underflow guard when ending parallel tasks

#rnx

diff --git a/Engine/Source/Developer/Localization/Private/LocalizationSourceControlUtil.cpp b/Engine/Source/Developer/Localization/Private/LocalizationSourceControlUtil.cpp
--- a/Engine/Source/Developer/Localization/Private/LocalizationSourceControlUtil.cpp
+++ b/Engine/Source/Developer/Localization/Private/LocalizationSourceControlUtil.cpp
@@ -45,7 +45,9 @@ void FLocalizationSCC::BeginParallelTasks()
 
 bool FLocalizationSCC::EndParallelTasks(FText& OutError)
 {
-	if (--ParallelTasksCount > 0)
+	const uint16 PrevParallelTasksCount = ParallelTasksCount--;
+	checkf(PrevParallelTasksCount > 0, TEXT("FLocalizationSCC::EndParallelTasks was called while no parallel tasks are running"));
+	if (PrevParallelTasksCount > 1)
 	{
 		return true;
 	}
diff --git a/Engine/Source/Developer/Localization/Public/LocalizationSourceControlUtil.h b/Engine/Source/Developer/Localization/Public/LocalizationSourceControlUtil.h
--- a/Engine/Source/Developer/Localization/Public/LocalizationSourceControlUtil.h
+++ b/Engine/Source/Developer/Localization/Public/LocalizationSourceControlUtil.h
@@ -43,7 +43,7 @@ private:
 	/** Mutex protecting CheckedOutFiles and DeferredCheckedOutFiles when ParallelTasksCount > 0 */
 	mutable UE::FMutex FilesMutex;
 	/** >0 if this we're being used from a parallel tasks block, and need to defer check-out requests */
-	std::atomic<uint8> ParallelTasksCount = 0;
+	std::atomic<uint16> ParallelTasksCount = 0;
 };
 
 class FLocFileSCCNotifies : public ILocFileNotifies
