Fixed MergeLocalization failing to correctly forward bUseSourceControl in a few places

#rnx

diff --git a/Engine/Source/Developer/Localization/Private/UserGeneratedContentLocalization.cpp b/Engine/Source/Developer/Localization/Private/UserGeneratedContentLocalization.cpp
--- a/Engine/Source/Developer/Localization/Private/UserGeneratedContentLocalization.cpp
+++ b/Engine/Source/Developer/Localization/Private/UserGeneratedContentLocalization.cpp
@@ -1031,10 +1031,10 @@ bool MergeLocalization(const FString& SourceLocalizationTargetName, const FStrin
 		else
 		{
 			// File doesn't exist at dest; just copy directly from source
-			PreWriteFileWithSCC(DestFilename);
+			PreWriteFileWithSCC(DestFilename, bUseSourceControl);
 			if (IFileManager::Get().Copy(*DestFilename, *SourceFilename) == COPY_OK)
 			{
-				PostWriteFileWithSCC(DestFilename);
+				PostWriteFileWithSCC(DestFilename, bUseSourceControl);
 				UE_LOG(LogUGCLocalization, Log, TEXT("Merge copied '%s' to '%s'"), *SourceFilename, *DestFilename);
 			}
 			else
