Avoid the modification (and thus checkout) of localization files that haven't changed on-disk

This avoids source control churn, where we checkout a file only to later revert it as it hasn't changed

diff --git a/Engine/Source/Developer/Localization/Private/LocTextHelper.cpp b/Engine/Source/Developer/Localization/Private/LocTextHelper.cpp
--- a/Engine/Source/Developer/Localization/Private/LocTextHelper.cpp
+++ b/Engine/Source/Developer/Localization/Private/LocTextHelper.cpp
@@ -5,6 +5,7 @@
 #include "Dom/JsonObject.h"
 #include "Internationalization/BreakIterator.h"
 #include "Internationalization/IBreakIterator.h"
+#include "LocalizationFileUtil.h"
 #include "Logging/StructuredLog.h"
 #include "Misc/DataDrivenPlatformInfoRegistry.h"
 #include "Misc/EnumClassFlags.h"
@@ -1254,12 +1255,6 @@ FString FLocTextHelper::GetConflictReport() const
 
 bool FLocTextHelper::SaveConflictReport(const FString& InReportFilePath, EConflictReportFormat InConflictReportFormat, FText* OutError) const
 {
-	bool bSaved = false;
-
-	if (LocFileNotifies.IsValid())
-	{
-		LocFileNotifies->PreFileWrite(InReportFilePath);
-	}
 	FString ConflictReport;
 	switch (InConflictReportFormat)
 	{
@@ -1283,21 +1278,15 @@ bool FLocTextHelper::SaveConflictReport(const FString& InReportFilePath, EConfli
 		}
 	}
 	
-	if (FFileHelper::SaveStringToFile(ConflictReport, *InReportFilePath))
-	{
-		bSaved = true;
-	}
-	else
-	{
-		if (OutError)
+	const bool bSaved = UE::Localization::FileUtil::WriteFileIfModified(InReportFilePath,
+		[&ConflictReport](const FString& Filename)
 		{
-			*OutError = FText::Format(LOCTEXT("Error_SaveConflictReport_SaveStringToFile", "Failed to save conflict report '{0}'."), FText::FromString(InReportFilePath));
-		}
-	}
+			return FFileHelper::SaveStringToFile(ConflictReport, *Filename);
+		}, &UE::Localization::FileUtil::TextHash, LocFileNotifies);
 
-	if (LocFileNotifies.IsValid())
+	if (!bSaved && OutError)
 	{
-		LocFileNotifies->PostFileWrite(InReportFilePath);
+		*OutError = FText::Format(LOCTEXT("Error_SaveConflictReport_SaveStringToFile", "Failed to save conflict report '{0}'."), FText::FromString(InReportFilePath));
 	}
 
 	return bSaved;
@@ -1446,32 +1435,18 @@ FLocTextWordCounts FLocTextHelper::GetWordCountReport(const FDateTime& InTimesta
 
 bool FLocTextHelper::SaveWordCountReport(const FDateTime& InTimestamp, const FString& InReportFilePath, FText* OutError) const
 {
-	bool bSaved = false;
-
-	if (LocFileNotifies.IsValid())
-	{
-		LocFileNotifies->PreFileWrite(InReportFilePath);
-	}
-
 	FLocTextWordCounts WordCounts = GetWordCountReport(InTimestamp, *InReportFilePath);
 	WordCounts.TrimReport();
 
-	const FString WordCountReportCSV = WordCounts.ToCSV();
-	if (FFileHelper::SaveStringToFile(WordCountReportCSV, *InReportFilePath))
-	{
-		bSaved = true;
-	}
-	else
-	{
-		if (OutError)
+	const bool bSaved = UE::Localization::FileUtil::WriteFileIfModified(InReportFilePath,
+		[&WordCounts](const FString& Filename)
 		{
-			*OutError = FText::Format(LOCTEXT("Error_SaveWordCountReport_SaveStringToFile", "Failed to save word count report '{0}'."), FText::FromString(InReportFilePath));
-		}
-	}
+			return FFileHelper::SaveStringToFile(WordCounts.ToCSV(), *Filename);
+		}, &UE::Localization::FileUtil::TextHash, LocFileNotifies);
 
-	if (LocFileNotifies.IsValid())
+	if (!bSaved && OutError)
 	{
-		LocFileNotifies->PostFileWrite(InReportFilePath);
+		*OutError = FText::Format(LOCTEXT("Error_SaveWordCountReport_SaveStringToFile", "Failed to save word count report '{0}'."), FText::FromString(InReportFilePath));
 	}
 
 	return bSaved;
@@ -1719,28 +1694,15 @@ bool FLocTextHelper::SaveManifestImpl(const TSharedRef<const FInternationalizati
 {
 	auto SaveSingleManifest = [this, &OutError](const TSharedRef<const FInternationalizationManifest>& InManifestToSave, const FString& InManifestFilePathToSave) -> bool
 	{
-		bool bSaved = false;
-
-		if (LocFileNotifies.IsValid())
-		{
-			LocFileNotifies->PreFileWrite(InManifestFilePathToSave);
-		}
-
-		if (FJsonInternationalizationManifestSerializer::SerializeManifestToFile(InManifestToSave, InManifestFilePathToSave))
-		{
-			bSaved = true;
-		}
-		else
-		{
-			if (OutError)
+		const bool bSaved = UE::Localization::FileUtil::WriteFileIfModified(InManifestFilePathToSave,
+			[&InManifestToSave](const FString& Filename)
 			{
-				*OutError = FText::Format(LOCTEXT("Error_SaveManifest_SerializeFile", "Failed to serialize manifest '{0}'."), FText::FromString(InManifestFilePathToSave));
-			}
-		}
+				return FJsonInternationalizationManifestSerializer::SerializeManifestToFile(InManifestToSave, Filename);
+			}, &UE::Localization::FileUtil::TextHash, LocFileNotifies);
 
-		if (LocFileNotifies.IsValid())
+		if (!bSaved && OutError)
 		{
-			LocFileNotifies->PostFileWrite(InManifestFilePathToSave);
+			*OutError = FText::Format(LOCTEXT("Error_SaveManifest_SerializeFile", "Failed to serialize manifest '{0}'."), FText::FromString(InManifestFilePathToSave));
 		}
 
 		return bSaved;
@@ -1913,28 +1875,15 @@ bool FLocTextHelper::SaveArchiveImpl(const TSharedRef<const FInternationalizatio
 {
 	auto SaveSingleArchive = [this, &OutError](const TSharedRef<const FInternationalizationArchive>& InArchiveToSave, const FString& InArchiveFilePathToSave) -> bool
 	{
-		bool bSaved = false;
-
-		if (LocFileNotifies.IsValid())
-		{
-			LocFileNotifies->PreFileWrite(InArchiveFilePathToSave);
-		}
-
-		if (FJsonInternationalizationArchiveSerializer::SerializeArchiveToFile(InArchiveToSave, InArchiveFilePathToSave))
-		{
-			bSaved = true;
-		}
-		else
-		{
-			if (OutError)
+		const bool bSaved = UE::Localization::FileUtil::WriteFileIfModified(InArchiveFilePathToSave,
+			[&InArchiveToSave](const FString& Filename)
 			{
-				*OutError = FText::Format(LOCTEXT("Error_SaveArchive_SerializeFile", "Failed to serialize archive '{0}'."), FText::FromString(InArchiveFilePathToSave));
-			}
-		}
+				return FJsonInternationalizationArchiveSerializer::SerializeArchiveToFile(InArchiveToSave, Filename);
+			}, &UE::Localization::FileUtil::TextHash, LocFileNotifies);
 
-		if (LocFileNotifies.IsValid())
+		if (!bSaved && OutError)
 		{
-			LocFileNotifies->PostFileWrite(InArchiveFilePathToSave);
+			*OutError = FText::Format(LOCTEXT("Error_SaveArchive_SerializeFile", "Failed to serialize archive '{0}'."), FText::FromString(InArchiveFilePathToSave));
 		}
 
 		return bSaved;
diff --git a/Engine/Source/Developer/Localization/Private/LocalizationFileUtil.cpp b/Engine/Source/Developer/Localization/Private/LocalizationFileUtil.cpp
new file mode 100644
--- /dev/null
+++ b/Engine/Source/Developer/Localization/Private/LocalizationFileUtil.cpp
@@ -0,0 +1,122 @@
+// Copyright Epic Games, Inc. All Rights Reserved.
+
+#include "LocalizationFileUtil.h"
+
+#include "Hash/xxhash.h"
+#include "HAL/FileManager.h"
+#include "LocalizationSourceControlUtil.h"
+#include "Memory/MemoryView.h"
+#include "Misc/FileHelper.h"
+#include "Misc/Paths.h"
+#include "Misc/ScopeExit.h"
+
+namespace UE::Localization::FileUtil::Private
+{
+	
+bool WriteFileIfModified(const FString& Filename, TFunctionRef<bool(const FString&)> WriteFunc, TFunctionRef<uint64(const FString&)> HashFunc, ILocFileNotifies* LocFileNotifies)
+{
+	const FString TempFilename = Filename + TEXTVIEW(".tmp");
+
+	if (!WriteFunc(TempFilename))
+	{
+		return false;
+	}
+
+	IFileManager& FileManager = IFileManager::Get();
+
+	ON_SCOPE_EXIT
+	{
+		FileManager.Delete(*TempFilename);
+	};
+
+	const uint64 CurrentHash = HashFunc(Filename);
+	const uint64 NewHash = HashFunc(TempFilename);
+
+	if (CurrentHash == NewHash)
+	{
+		return true;
+	}
+
+	if (LocFileNotifies)
+	{
+		LocFileNotifies->PreFileWrite(Filename);
+	}
+
+	const bool bUpdatedFile = FileManager.Move(*Filename, *TempFilename);
+	
+	if (LocFileNotifies)
+	{
+		LocFileNotifies->PostFileWrite(Filename);
+	}
+	
+	return bUpdatedFile;
+}
+	
+}
+
+uint64 UE::Localization::FileUtil::BinaryHash(const FString& Filename)
+{
+	FXxHash64Builder HashBuilder;
+	if (FPaths::FileExists(Filename))
+	{
+		FFileHelper::LoadFileInBlocks(Filename, 
+			[&HashBuilder](FMemoryView Block)
+			{
+				HashBuilder.Update(Block);
+			});
+	}
+	return HashBuilder.Finalize().Hash;
+}
+
+uint64 UE::Localization::FileUtil::TextHash(const FString& Filename)
+{
+	FXxHash64Builder HashBuilder;
+	if (FPaths::FileExists(Filename))
+	{
+		FFileHelper::LoadFileToStringWithLineVisitor(*Filename,
+			[&HashBuilder](FStringView Line)
+			{
+				HashBuilder.Update(Line.GetData(), Line.Len() * sizeof(TCHAR));
+			});
+	}
+	return HashBuilder.Finalize().Hash;
+}
+
+uint64 UE::Localization::FileUtil::PortableObjectHash(const FString& Filename)
+{
+	FXxHash64Builder HashBuilder;
+	if (FPaths::FileExists(Filename))
+	{
+		// Don't include the PO file header in the hash as it contains transient information (like timestamps) that we don't care about
+		bool bSkippingPOFileHeader = true;
+
+		FFileHelper::LoadFileToStringWithLineVisitor(*Filename,
+			[&bSkippingPOFileHeader, &HashBuilder](FStringView Line)
+			{
+				if (bSkippingPOFileHeader)
+				{
+					// PO file headers end on the first empty line in the file
+					bSkippingPOFileHeader = !Line.IsEmpty();
+					return;
+				}
+				HashBuilder.Update(Line.GetData(), Line.Len() * sizeof(TCHAR));
+			});
+	}
+	return HashBuilder.Finalize().Hash;
+}
+
+bool UE::Localization::FileUtil::WriteFileIfModified(const FString& Filename, TFunctionRef<bool(const FString&)> WriteFunc, TFunctionRef<uint64(const FString&)> HashFunc)
+{
+	return Private::WriteFileIfModified(Filename, WriteFunc, HashFunc, nullptr);
+}
+
+bool UE::Localization::FileUtil::WriteFileIfModified(const FString& Filename, TFunctionRef<bool(const FString&)> WriteFunc, TFunctionRef<uint64(const FString&)> HashFunc, const TSharedPtr<FLocalizationSCC>& SourceControlInfo)
+{
+	FLocFileSCCNotifies LocFileNotifies{SourceControlInfo};
+	return Private::WriteFileIfModified(Filename, WriteFunc, HashFunc, &LocFileNotifies);
+}
+
+bool UE::Localization::FileUtil::WriteFileIfModified(const FString& Filename, TFunctionRef<bool(const FString&)> WriteFunc, TFunctionRef<uint64(const FString&)> HashFunc, const TSharedPtr<ILocFileNotifies>& LocFileNotifies)
+{
+	return Private::WriteFileIfModified(Filename, WriteFunc, HashFunc, LocFileNotifies.Get());
+}
diff --git a/Engine/Source/Developer/Localization/Private/PortableObjectPipeline.cpp b/Engine/Source/Developer/Localization/Private/PortableObjectPipeline.cpp
--- a/Engine/Source/Developer/Localization/Private/PortableObjectPipeline.cpp
+++ b/Engine/Source/Developer/Localization/Private/PortableObjectPipeline.cpp
@@ -9,6 +9,7 @@
 #include "Serialization/JsonInternationalizationMetadataSerializer.h"
 #include "PortableObjectFormatDOM.h"
 #include "Internationalization/TextNamespaceUtil.h"
+#include "LocalizationFileUtil.h"
 #include "LocTextHelper.h"
 
 #include UE_INLINE_GENERATED_CPP_BY_NAME(PortableObjectPipeline)
@@ -419,24 +420,12 @@ namespace
 
 		NewPortableObject.SortEntries();
 
-		bool bPOFileSaved = false;
-		{
-			TSharedPtr<ILocFileNotifies> LocFileNotifies = InLocTextHelper.GetLocFileNotifies();
-
-			if (LocFileNotifies.IsValid())
-			{
-				LocFileNotifies->PreFileWrite(InPOFilePath);
-			}
-
-			//@TODO We force UTF8 at the moment but we want this to be based on the format found in the header info.
-			const FString OutputString = NewPortableObject.ToString();
-			bPOFileSaved = FFileHelper::SaveStringToFile(OutputString, *InPOFilePath, FFileHelper::EEncodingOptions::ForceUTF8);
-
-			if (LocFileNotifies.IsValid())
+		const bool bPOFileSaved = UE::Localization::FileUtil::WriteFileIfModified(InPOFilePath,
+			[&NewPortableObject](const FString& Filename)
 			{
-				LocFileNotifies->PostFileWrite(InPOFilePath);
-			}
-		}
+				//@TODO We force UTF8 at the moment but we want this to be based on the format found in the header info.
+				return FFileHelper::SaveStringToFile(NewPortableObject.ToString(), *Filename, FFileHelper::EEncodingOptions::ForceUTF8);
+			}, &UE::Localization::FileUtil::PortableObjectHash, InLocTextHelper.GetLocFileNotifies());
 
 		if (!bPOFileSaved)
 		{
diff --git a/Engine/Source/Developer/Localization/Private/UserGeneratedContentLocalization.cpp b/Engine/Source/Developer/Localization/Private/UserGeneratedContentLocalization.cpp
--- a/Engine/Source/Developer/Localization/Private/UserGeneratedContentLocalization.cpp
+++ b/Engine/Source/Developer/Localization/Private/UserGeneratedContentLocalization.cpp
@@ -20,6 +20,7 @@
 #include "TextLocalizationResourceGenerator.h"
 #include "LocalizationConfigurationScript.h"
 #include "LocalizationDelegates.h"
+#include "LocalizationFileUtil.h"
 #include "LocTextHelper.h"
 #include "PortableObjectFormatDOM.h"
 #include "SourceControlHelpers.h"
@@ -281,10 +282,10 @@ void PostWriteFileWithSCC(const FString& Filename, const bool bUseSourceControl)
 	}
 }
 
-void GetLocalizationFileHashes(const FString& BaseDirectory, TMap<FString, FXxHash64>& OutFileHashes)
+void GetLocalizationFileHashes(const FString& BaseDirectory, TMap<FString, uint64>& OutFileHashes)
 {
 	// Get the list of localization files under the given directory
-	TArray<TTuple<FString, FXxHash64>> FilesToHash;
+	TArray<TTuple<FString, uint64>> FilesToHash;
 	IFileManager::Get().IterateDirectoryRecursively(*BaseDirectory, [&FilesToHash](const TCHAR* FilenameOrDirectory, bool bIsDirectory)
 	{
 		if (!bIsDirectory)
@@ -292,7 +293,7 @@ void GetLocalizationFileHashes(const FString& BaseDirectory, TMap<FString, FXxHa
 			if (FStringView FileExtension = FPathViews::GetExtension(FilenameOrDirectory);
 				FileExtension == TEXTVIEW("po") || FileExtension == TEXTVIEW("ugcloc"))
 			{
-				FilesToHash.Add(MakeTuple(FilenameOrDirectory, FXxHash64()));
+				FilesToHash.Add(MakeTuple(FilenameOrDirectory, 0));
 			}
 		}
 		return true;
@@ -301,30 +302,16 @@ void GetLocalizationFileHashes(const FString& BaseDirectory, TMap<FString, FXxHa
 	// Generate their content hashes in parallel
 	ParallelFor(TEXT("UserGeneratedContentLocalization.GetLocalizationFileHashes"), FilesToHash.Num(), 1, [&FilesToHash](int32 Index)
 	{
-		TTuple<FString, FXxHash64>& FileToHashPair = FilesToHash[Index];
+		TTuple<FString, uint64>& FileToHashPair = FilesToHash[Index];
 
-		FXxHash64Builder HashBuilder;
-		{
-			// Don't include the PO file header in the hash as it contains transient information (like timestamps) that we don't care about
-			bool bSkippingPOFileHeader = FPathViews::GetExtension(FileToHashPair.Key) == TEXTVIEW("po");
-
-			FFileHelper::LoadFileToStringWithLineVisitor(*FileToHashPair.Key, [&bSkippingPOFileHeader, &HashBuilder](FStringView Line)
-			{
-				if (bSkippingPOFileHeader)
-				{
-					// PO file headers end on the first empty line in the file
-					bSkippingPOFileHeader = !Line.IsEmpty();
-					return;
-				}
-				HashBuilder.Update(Line.GetData(), Line.Len() * sizeof(TCHAR));
-			});
-		}
-		FileToHashPair.Value = HashBuilder.Finalize();
+		FileToHashPair.Value = FPathViews::GetExtension(FileToHashPair.Key) == TEXTVIEW("po")
+			? UE::Localization::FileUtil::PortableObjectHash(FileToHashPair.Key)
+			: UE::Localization::FileUtil::TextHash(FileToHashPair.Key);
 	});
 
 	// Append the content hashes to the result
 	OutFileHashes.Reserve(OutFileHashes.Num() + FilesToHash.Num());
-	for (TTuple<FString, FXxHash64>& FileToHashPair : FilesToHash)
+	for (TTuple<FString, uint64>& FileToHashPair : FilesToHash)
 	{
 		OutFileHashes.Add(MoveTemp(FileToHashPair.Key), MoveTemp(FileToHashPair.Value));
 	}
@@ -356,7 +343,7 @@ bool ExportLocalization(TArrayView<const TSharedRef<IPlugin>> Plugins, const FEx
 
 	// Localization data stored per-plugin
 	TArray<FString, TInlineAllocator<1>> GatherConfigFilenames;
-	TMap<FString, TMap<FString, FXxHash64>> PerTargetLocalizationFileHashes;
+	TMap<FString, TMap<FString, uint64>> PerTargetLocalizationFileHashes;
 	for (const TSharedRef<IPlugin>& Plugin : Plugins)
 	{
 		const FString PluginLocalizationTargetName = GetLocalizationTargetName(Plugin);
@@ -364,11 +351,11 @@ bool ExportLocalization(TArrayView<const TSharedRef<IPlugin>> Plugins, const FEx
 		const FString PluginLocalizationTargetDirectory = GetLocalizationTargetDirectory(PluginLocalizationTargetName, Plugin->GetContentDir(), ExportOptions.LocalizationTargetRootDirectoryOverride);
 
 		// Track the source file hashes when the export started, so that we can detect post-export whether the files have actually changed
-		TMap<FString, FXxHash64>& LocalizationFileHashes = PerTargetLocalizationFileHashes.FindOrAdd(PluginLocalizationTargetName);
+		TMap<FString, uint64>& LocalizationFileHashes = PerTargetLocalizationFileHashes.FindOrAdd(PluginLocalizationTargetName);
 		GetLocalizationFileHashes(PluginLocalizationTargetDirectory, LocalizationFileHashes);
 
 		// Seed the scratch directory with the current localization files for this plugin, so that the loc gather will import and preserve any existing translation data
-		for (const TTuple<FString, FXxHash64>& LocalizationFileHashPair : LocalizationFileHashes)
+		for (const TTuple<FString, uint64>& LocalizationFileHashPair : LocalizationFileHashes)
 		{
 			const FString& SourceFilename = LocalizationFileHashPair.Key;
 			FString DestinationFilename = SourceFilename;
@@ -635,18 +622,18 @@ bool ExportLocalization(TArrayView<const TSharedRef<IPlugin>> Plugins, const FEx
 		}
 
 		// Track the scratch file hashes now that the export has finished, so that we can detect which files have actually changed from their source
-		TMap<FString, FXxHash64> ScratchLocalizationFileHashes;
+		TMap<FString, uint64> ScratchLocalizationFileHashes;
 		GetLocalizationFileHashes(PluginLocalizationScratchDirectory, ScratchLocalizationFileHashes);
 
-		const TMap<FString, FXxHash64>& SourceLocalizationFileHashes = PerTargetLocalizationFileHashes.FindChecked(PluginLocalizationTargetName);
-		for (const TTuple<FString, FXxHash64>& LocalizationFileHashPair : ScratchLocalizationFileHashes)
+		const TMap<FString, uint64>& SourceLocalizationFileHashes = PerTargetLocalizationFileHashes.FindChecked(PluginLocalizationTargetName);
+		for (const TTuple<FString, uint64>& LocalizationFileHashPair : ScratchLocalizationFileHashes)
 		{
 			const FString& SourceFilename = LocalizationFileHashPair.Key;
 			FString DestinationFilename = SourceFilename;
 			if (DestinationFilename.ReplaceInline(*PluginLocalizationScratchDirectory, *PluginLocalizationTargetDirectory) > 0)
 			{
 				// Only copy files with modified hashes to avoid source control churn
-				if (const FXxHash64* SourceFileHash = SourceLocalizationFileHashes.Find(DestinationFilename);
+				if (const uint64* SourceFileHash = SourceLocalizationFileHashes.Find(DestinationFilename);
 					!SourceFileHash || *SourceFileHash != LocalizationFileHashPair.Value)
 				{
 					PreWriteFileWithSCC(DestinationFilename, bUseSourceControl);
@@ -952,7 +939,7 @@ ELoadLocalizationResult LoadLocalization(const FString& LocalizationTargetName,
 
 bool MergeLocalization(const FString& SourceLocalizationTargetName, const FString& SourceLocalizationTargetDirectory, const FString& DestLocalizationTargetName, const FString& DestLocalizationTargetDirectory, const EMergeLocalizationMode MergeMode, const bool bUseSourceControl)
 {
-	TMap<FString, FXxHash64> SourceLocalizationFileHashes;
+	TMap<FString, uint64> SourceLocalizationFileHashes;
 	GetLocalizationFileHashes(SourceLocalizationTargetDirectory, SourceLocalizationFileHashes);
 
 	if (SourceLocalizationFileHashes.IsEmpty())
@@ -961,7 +948,7 @@ bool MergeLocalization(const FString& SourceLocalizationTargetName, const FStrin
 		return true;
 	}
 
-	TMap<FString, FXxHash64> DestLocalizationFileHashes;
+	TMap<FString, uint64> DestLocalizationFileHashes;
 	GetLocalizationFileHashes(DestLocalizationTargetDirectory, DestLocalizationFileHashes);
 
 	auto SourceFilenameToDestFilename = [&SourceLocalizationTargetDirectory, &DestLocalizationTargetDirectory, SourcePOFilename = SourceLocalizationTargetName + TEXT(".po"), DestPOFilename = DestLocalizationTargetName + TEXT(".po")](const FString& SourceFilename)
@@ -1010,12 +997,12 @@ bool MergeLocalization(const FString& SourceLocalizationTargetName, const FStrin
 	};
 
 	TArray<FPOFileToMerge> POFilesToMerge;
-	for (const TTuple<FString, FXxHash64>& SourceLocalizationFileHashPair : SourceLocalizationFileHashes)
+	for (const TTuple<FString, uint64>& SourceLocalizationFileHashPair : SourceLocalizationFileHashes)
 	{
 		const FString& SourceFilename = SourceLocalizationFileHashPair.Key;
 		const FString DestFilename = SourceFilenameToDestFilename(SourceFilename);
 
-		if (const FXxHash64* DestLocalizationFileHash = DestLocalizationFileHashes.Find(DestFilename))
+		if (const uint64* DestLocalizationFileHash = DestLocalizationFileHashes.Find(DestFilename))
 		{
 			if (SourceLocalizationFileHashPair.Value == *DestLocalizationFileHash)
 			{
diff --git a/Engine/Source/Developer/Localization/Public/LocalizationFileUtil.h b/Engine/Source/Developer/Localization/Public/LocalizationFileUtil.h
new file mode 100644
--- /dev/null
+++ b/Engine/Source/Developer/Localization/Public/LocalizationFileUtil.h
@@ -0,0 +1,47 @@
+// Copyright Epic Games, Inc. All Rights Reserved.
+
+#pragma once 
+
+#include "HAL/Platform.h"
+#include "Templates/FunctionFwd.h"
+#include "Templates/SharedPointerFwd.h"
+
+#define UE_API LOCALIZATION_API
+
+class FLocalizationSCC;
+class FString;
+class ILocFileNotifies;
+
+namespace UE::Localization::FileUtil
+{
+	
+/**
+ * Generic file hash, that hashes the binary representation of the file.
+ */
+UE_API uint64 BinaryHash(const FString& Filename);
+
+/**
+ * Generic file hash, that hashes the textual representation of the file (ignoring line endings and encoding).
+ */
+UE_API uint64 TextHash(const FString& Filename);
+
+/**
+ * Specialized file hash for PO files, that skips the transient information (like timestamps) from the PO file header.
+ */
+UE_API uint64 PortableObjectHash(const FString& Filename);
+
+/**
+ * Utility to write a file to disk, but only if it has changed.
+ * 
+ * This function works by using WriteFunc to generate a .tmp file and then uses HashFunc to compare the contents of the .tmp file against the contents of Filename (if it exists).
+ * If the hash is different, then the file is atomically renamed to Filename, which ensures that we don't get a partial file write if the application crashes.
+ * 
+ * @return True if the file was modified or up-to-date, or false if the file failed to write.
+ */
+UE_API bool WriteFileIfModified(const FString& Filename, TFunctionRef<bool(const FString&)> WriteFunc, TFunctionRef<uint64(const FString&)> HashFunc);
+UE_API bool WriteFileIfModified(const FString& Filename, TFunctionRef<bool(const FString&)> WriteFunc, TFunctionRef<uint64(const FString&)> HashFunc, const TSharedPtr<FLocalizationSCC>& SourceControlInfo);
+UE_API bool WriteFileIfModified(const FString& Filename, TFunctionRef<bool(const FString&)> WriteFunc, TFunctionRef<uint64(const FString&)> HashFunc, const TSharedPtr<ILocFileNotifies>& LocFileNotifies);
+	
+}
+
+#undef UE_API
diff --git a/Engine/Source/Editor/UnrealEd/Private/Commandlets/GenerateTextLocalizationResourceCommandlet.cpp b/Engine/Source/Editor/UnrealEd/Private/Commandlets/GenerateTextLocalizationResourceCommandlet.cpp
--- a/Engine/Source/Editor/UnrealEd/Private/Commandlets/GenerateTextLocalizationResourceCommandlet.cpp
+++ b/Engine/Source/Editor/UnrealEd/Private/Commandlets/GenerateTextLocalizationResourceCommandlet.cpp
@@ -11,8 +11,8 @@
 #include "Internationalization/Text.h"
 #include "Internationalization/TextKey.h"
 #include "Internationalization/TextLocalizationResource.h"
+#include "LocalizationFileUtil.h"
 #include "LocTextHelper.h"
-#include "LocalizedAssetUtil.h"
 #include "Logging/LogCategory.h"
 #include "Logging/LogMacros.h"
 #include "Misc/Paths.h"
@@ -181,11 +181,12 @@ int32 UGenerateTextLocalizationResourceCommandlet::Main(const FString& Params)
 	{
 		const FString TextLocalizationMetaDataResourcePath = DestinationPath / FPaths::GetBaseFilename(ResourceName) + TEXT(".locmeta");
 
-		const bool bLocMetaFileSaved = FLocalizedAssetSCCUtil::SaveFileWithSCC(SourceControlInfo, TextLocalizationMetaDataResourcePath, [&LocTextHelper, &ResourceName](const FString& InSaveFileName) -> bool
-		{
-			FTextLocalizationMetaDataResource LocMeta;
-			return FTextLocalizationResourceGenerator::GenerateLocMeta(LocTextHelper, ResourceName, LocMeta) && LocMeta.SaveToFile(InSaveFileName);
-		});
+		const bool bLocMetaFileSaved = UE::Localization::FileUtil::WriteFileIfModified(TextLocalizationMetaDataResourcePath,
+			[&LocTextHelper, &ResourceName](const FString& Filename)
+			{
+				FTextLocalizationMetaDataResource LocMeta;
+				return FTextLocalizationResourceGenerator::GenerateLocMeta(LocTextHelper, ResourceName, LocMeta) && LocMeta.SaveToFile(Filename);
+			}, &UE::Localization::FileUtil::BinaryHash, SourceControlInfo);
 
 		if (bLocMetaFileSaved)
 		{
@@ -223,10 +224,11 @@ int32 UGenerateTextLocalizationResourceCommandlet::Main(const FString& Params)
 				TextLocalizationResourcePath = DestinationPath / FPaths::GetPlatformLocalizationFolderName() / InPlatformName.ToString() / CultureName / ResourceName;
 			}
 
-			const bool bLocResFileSaved = FLocalizedAssetSCCUtil::SaveFileWithSCC(SourceControlInfo, TextLocalizationResourcePath, [&InLocRes](const FString& InSaveFileName) -> bool
-			{
-				return InLocRes.SaveToFile(InSaveFileName);
-			});
+			const bool bLocResFileSaved = UE::Localization::FileUtil::WriteFileIfModified(TextLocalizationResourcePath,
+				[&InLocRes](const FString& Filename)
+				{
+					return InLocRes.SaveToFile(Filename);
+				}, &UE::Localization::FileUtil::BinaryHash, SourceControlInfo);
 
 			if (bLocResFileSaved)
 			{
diff --git a/Engine/Source/Programs/AutomationTool/Scripts/Localisation.Automation.cs b/Engine/Source/Programs/AutomationTool/Scripts/Localisation.Automation.cs
--- a/Engine/Source/Programs/AutomationTool/Scripts/Localisation.Automation.cs
+++ b/Engine/Source/Programs/AutomationTool/Scripts/Localisation.Automation.cs
@@ -356,8 +356,9 @@ class Localize : BuildCommand
 		// Clean-up the changelist so it only contains the changed files, and then submit it (if we were asked to)
 		if (P4Enabled && !bIsRunningInPreview)
 		{
-			// Revert any PO files that haven't changed aside from their header
-			RevertUnchangedFiles(LocalizationBatches, InitalPOFileHashes);
+			// Revert any PO files that haven't changed aside from their header, 
+			// and ensure that any that have are added/edited in source control
+			FinalizeSourceControlState(LocalizationBatches, InitalPOFileHashes);
 
 			// Submit that single changelist now
 			if (AllowSubmit)
@@ -917,7 +918,7 @@ class Localize : BuildCommand
 		}
 	}
 
-	private void RevertUnchangedFiles(List<LocalizationBatch> LocalizationBatches, Dictionary<string, Byte[]> InitalPOFileHashes)
+	private void FinalizeSourceControlState(List<LocalizationBatch> LocalizationBatches, Dictionary<string, Byte[]> InitalPOFileHashes)
 	{
 		if (!P4Enabled || bIsRunningInPreview)
 		{
@@ -925,6 +926,8 @@ class Localize : BuildCommand
 		}
 
 		{
+			var POFilesToAdd = new List<string>();
+			var POFilesToEdit = new List<string>();
 			var POFilesToRevert = new List<string>();
 
 			var CurrentPOFileHashes = GetPOFileHashes(LocalizationBatches, UEProjectRoot);
@@ -935,22 +938,29 @@ class Localize : BuildCommand
 				{
 					POFilesToRevert.Add(CurrentPOFileHashPair.Key);
 				}
+				else if (P4.FileExistsInDepot(CurrentPOFileHashPair.Key))
+				{
+					POFilesToEdit.Add(CurrentPOFileHashPair.Key);
+				}
+				else
+				{
+					POFilesToAdd.Add(CurrentPOFileHashPair.Key);
+				}
 			}
 
-			if (POFilesToRevert.Count > 0)
+			if (POFilesToAdd.Count > 0)
 			{
-				var P4RevertArgsFilename = CombinePaths(CmdEnv.LocalRoot, "Engine", "Intermediate", String.Format("LocalizationP4RevertArgs-{0}.txt", Guid.NewGuid().ToString()));
+				P4.Add(PendingChangeList, POFilesToAdd);
+			}
 
-				using (StreamWriter P4RevertArgsWriter = File.CreateText(P4RevertArgsFilename))
-				{
-					foreach (var POFileToRevert in POFilesToRevert)
-					{
-						P4RevertArgsWriter.WriteLine(POFileToRevert);
-					}
-				}
+			if (POFilesToEdit.Count > 0)
+			{
+				P4.Edit(PendingChangeList, POFilesToEdit);
+			}
 
-				P4.LogP4(String.Format("-x {0}", P4RevertArgsFilename), "revert");
-				DeleteFile_NoExceptions(P4RevertArgsFilename);
+			if (POFilesToRevert.Count > 0)
+			{
+				P4.Revert(PendingChangeList, POFilesToRevert);
 			}
 		}
 
