Allow FLocalizationSCC to be created without force initializing the SCC provider (eg, when using it within a standard, non-commandlet, editor instance)

#rnx

diff --git a/Engine/Source/Developer/Localization/Private/LocalizationSourceControlUtil.cpp b/Engine/Source/Developer/Localization/Private/LocalizationSourceControlUtil.cpp
--- a/Engine/Source/Developer/Localization/Private/LocalizationSourceControlUtil.cpp
+++ b/Engine/Source/Developer/Localization/Private/LocalizationSourceControlUtil.cpp
@@ -19,11 +19,15 @@ namespace LocalizationSourceControlUtil
 
 #define LOCTEXT_NAMESPACE "LocalizationSourceControl"
 
-FLocalizationSCC::FLocalizationSCC()
+FLocalizationSCC::FLocalizationSCC(const EInitSCC InitSCC)
 {
 	checkf(IsInGameThread(), TEXT("FLocalizationSCC must be created on the game-thread"));
 
-	ISourceControlModule::Get().GetProvider().Init();
+	if (InitSCC == EInitSCC::Always || (InitSCC == EInitSCC::Commandlet && IsRunningCommandlet()))
+	{
+		ISourceControlModule::Get().GetProvider().Init();
+		bDidInitSCC = true;
+	}
 }
 
 FLocalizationSCC::~FLocalizationSCC()
@@ -35,7 +39,22 @@ FLocalizationSCC::~FLocalizationSCC()
 		UE_LOG(LogLocalizationSourceControl, Log, TEXT("Revision Control wrapper shutting down with checked out files."));
 	}
 
-	ISourceControlModule::Get().GetProvider().Close();
+	if (bDidInitSCC)
+	{
+		ISourceControlModule::Get().GetProvider().Close();
+	}
+}
+
+TSharedPtr<FLocalizationSCC> FLocalizationSCC::CreateIfSCCEnabled(const EInitSCC InitSCC)
+{
+	checkf(IsInGameThread(), TEXT("FLocalizationSCC::CreateIfSCCEnabled must be called on the game-thread"));
+
+	if (ISourceControlModule::Get().IsEnabled())
+	{
+		return MakeShared<FLocalizationSCC>(InitSCC);
+	}
+	
+	return nullptr;
 }
 
 void FLocalizationSCC::BeginParallelTasks()
diff --git a/Engine/Source/Developer/Localization/Public/LocalizationSourceControlUtil.h b/Engine/Source/Developer/Localization/Public/LocalizationSourceControlUtil.h
--- a/Engine/Source/Developer/Localization/Public/LocalizationSourceControlUtil.h
+++ b/Engine/Source/Developer/Localization/Public/LocalizationSourceControlUtil.h
@@ -16,9 +16,19 @@ class FText;
 class FLocalizationSCC
 {
 public:
-	UE_API FLocalizationSCC();
+	enum class EInitSCC
+	{
+		Never,
+		Always,
+		Commandlet,
+	};
+
+	UE_API explicit FLocalizationSCC(const EInitSCC InitSCC = EInitSCC::Commandlet);
 	UE_API ~FLocalizationSCC();
 
+	/** Create a FLocalizationSCC, but only if a SCC provider is enabled */
+	static UE_API TSharedPtr<FLocalizationSCC> CreateIfSCCEnabled(const EInitSCC InitSCC);
+
 	/** Start a block of parallel tasks (may be nested); this will defer check-out requests until the final EndParallelTasks is called */
 	UE_API void BeginParallelTasks();
 	/** Stop a block of parallel tasks; this will attempt to perform any check-outs deferred during the block */
@@ -44,6 +54,8 @@ private:
 	mutable UE::FMutex FilesMutex;
 	/** >0 if this we're being used from a parallel tasks block, and need to defer check-out requests */
 	std::atomic<uint16> ParallelTasksCount = 0;
+	/** True if we called Init on the SCC during construction, and need to call Close during destruction */
+	bool bDidInitSCC = false;
 };
 
 class FLocFileSCCNotifies : public ILocFileNotifies
