Updated UGC localization to use batched SCC operations

Using FLocFileSCCNotifies with an optional FLocalizationSCC allows us to easily batch the SCC operations, including allowing us to batch them directly from within a ParallelFor (when using BeginParallelTasks/EndParallelTasks).

#rnx

diff --git a/Engine/Source/Developer/Localization/Private/UserGeneratedContentLocalization.cpp b/Engine/Source/Developer/Localization/Private/UserGeneratedContentLocalization.cpp
--- a/Engine/Source/Developer/Localization/Private/UserGeneratedContentLocalization.cpp
+++ b/Engine/Source/Developer/Localization/Private/UserGeneratedContentLocalization.cpp
@@ -21,6 +21,7 @@
 #include "LocalizationConfigurationScript.h"
 #include "LocalizationDelegates.h"
 #include "LocalizationFileUtil.h"
+#include "LocalizationSourceControlUtil.h"
 #include "LocTextHelper.h"
 #include "PortableObjectFormatDOM.h"
 #include "SourceControlHelpers.h"
@@ -197,9 +198,11 @@ bool SetLocalizationTargetName(const TSharedRef<IPlugin>& Plugin, const FString&
 		PluginDescriptor.AdditionalFieldsToRemove.Emplace(GetLocalizationTargetNameFieldName());
 	}
 
-	PreWriteFileWithSCC(Plugin->GetDescriptorFileName(), bUseSourceControl);
+	FLocFileSCCNotifies LocFileNotifies = CreateLocFileNotifies(bUseSourceControl);
+
+	LocFileNotifies.PreFileWrite(Plugin->GetDescriptorFileName());
 	const bool bResult = Plugin->UpdateDescriptor(PluginDescriptor, OutFailReason);
-	PostWriteFileWithSCC(Plugin->GetDescriptorFileName(), bUseSourceControl);
+	LocFileNotifies.PostFileWrite(Plugin->GetDescriptorFileName());
 	return bResult;
 }
 
@@ -282,6 +285,16 @@ void PostWriteFileWithSCC(const FString& Filename, const bool bUseSourceControl)
 	}
 }
 
+FLocFileSCCNotifies CreateLocFileNotifies(const bool bUseSourceControl)
+{
+	TSharedPtr<FLocalizationSCC> SourceControlInfo;
+	if (bUseSourceControl)
+	{
+		SourceControlInfo = FLocalizationSCC::CreateIfSCCEnabled(FLocalizationSCC::EInitSCC::Never);
+	}
+	return FLocFileSCCNotifies{ SourceControlInfo };
+}
+
 void GetLocalizationFileHashes(const FString& BaseDirectory, TMap<FString, uint64>& OutFileHashes)
 {
 	// Get the list of localization files under the given directory
@@ -604,6 +617,14 @@ bool ExportLocalization(TArrayView<const TSharedRef<IPlugin>> Plugins, const FEx
 		}
 	}
 
+	FLocFileSCCNotifies LocFileNotifies = CreateLocFileNotifies(ExportOptions.LocalizationTargetRootDirectoryOverride.IsEmpty());
+
+	LocFileNotifies.BeginParallelTasks(); // Batch the below SCC operations
+	ON_SCOPE_EXIT
+	{
+		LocFileNotifies.EndParallelTasks(); // End the SCC batch
+	};
+
 	// Copy any updated PO files back to the plugins and write out the localization settings used to generate them
 	for (const TSharedRef<IPlugin>& Plugin : Plugins)
 	{
@@ -611,8 +632,6 @@ bool ExportLocalization(TArrayView<const TSharedRef<IPlugin>> Plugins, const FEx
 		const FString PluginLocalizationScratchDirectory = GetLocalizationScratchDirectory(PluginLocalizationTargetName);
 		const FString PluginLocalizationTargetDirectory = GetLocalizationTargetDirectory(PluginLocalizationTargetName, Plugin->GetContentDir(), ExportOptions.LocalizationTargetRootDirectoryOverride);
 
-		const bool bUseSourceControl = ExportOptions.LocalizationTargetRootDirectoryOverride.IsEmpty();
-
 		// Write the UGC localization descriptor settings that were used to generate this localization data
 		// That will be needed to handle compilation, but also to handle import correctly if the descriptor settings change later
 		if (const FString UGCLocFilename = GetLocalizationTargetUGCLocFile(PluginLocalizationTargetName, PluginLocalizationScratchDirectory);
@@ -636,10 +655,10 @@ bool ExportLocalization(TArrayView<const TSharedRef<IPlugin>> Plugins, const FEx
 				if (const uint64* SourceFileHash = SourceLocalizationFileHashes.Find(DestinationFilename);
 					!SourceFileHash || *SourceFileHash != LocalizationFileHashPair.Value)
 				{
-					PreWriteFileWithSCC(DestinationFilename, bUseSourceControl);
+					LocFileNotifies.PreFileWrite(DestinationFilename);
 					if (IFileManager::Get().Copy(*DestinationFilename, *SourceFilename) == COPY_OK)
 					{
-						PostWriteFileWithSCC(DestinationFilename, bUseSourceControl);
+						LocFileNotifies.PostFileWrite(DestinationFilename);
 						UE_LOG(LogUGCLocalization, Display, TEXT("Updated file for '%s': %s"), *PluginLocalizationTargetName, *DestinationFilename);
 					}
 					else
@@ -677,10 +696,10 @@ bool ExportLocalization(TArrayView<const TSharedRef<IPlugin>> Plugins, const FEx
 				}
 
 				FText DescriptorUpdateFailureReason;
-				PreWriteFileWithSCC(Plugin->GetDescriptorFileName(), bUseSourceControl);
+				LocFileNotifies.PreFileWrite(Plugin->GetDescriptorFileName());
 				if (Plugin->UpdateDescriptor(PluginDescriptor, DescriptorUpdateFailureReason))
 				{
-					PostWriteFileWithSCC(Plugin->GetDescriptorFileName(), bUseSourceControl);
+					LocFileNotifies.PostFileWrite(Plugin->GetDescriptorFileName());
 					UE_LOG(LogUGCLocalization, Display, TEXT("Updated .uplugin file for '%s'"), *PluginLocalizationTargetName);
 				}
 				else
@@ -951,6 +970,8 @@ bool MergeLocalization(const FString& SourceLocalizationTargetName, const FStrin
 	TMap<FString, uint64> DestLocalizationFileHashes;
 	GetLocalizationFileHashes(DestLocalizationTargetDirectory, DestLocalizationFileHashes);
 
+	FLocFileSCCNotifies LocFileNotifies = CreateLocFileNotifies(bUseSourceControl);
+
 	auto SourceFilenameToDestFilename = [&SourceLocalizationTargetDirectory, &DestLocalizationTargetDirectory, SourcePOFilename = SourceLocalizationTargetName + TEXT(".po"), DestPOFilename = DestLocalizationTargetName + TEXT(".po")](const FString& SourceFilename)
 	{
 		FString DestFilename = SourceFilename;
@@ -970,15 +991,15 @@ bool MergeLocalization(const FString& SourceLocalizationTargetName, const FStrin
 		return true;
 	};
 
-	auto SavePOFile = [bUseSourceControl](const FString& POFilename, FPortableObjectFormatDOM& POFile)
+	auto SavePOFile = [&LocFileNotifies](const FString& POFilename, FPortableObjectFormatDOM& POFile)
 	{
 		FString POFileData;
 		POFile.ToString(POFileData);
 
-		PreWriteFileWithSCC(POFilename, bUseSourceControl);
+		LocFileNotifies.PreFileWrite(POFilename);
 		if (FFileHelper::SaveStringToFile(POFileData, *POFilename, FFileHelper::EEncodingOptions::ForceUTF8))
 		{
-			PostWriteFileWithSCC(POFilename, bUseSourceControl);
+			LocFileNotifies.PostFileWrite(POFilename);
 			return true;
 		}
 		else
@@ -992,10 +1013,15 @@ bool MergeLocalization(const FString& SourceLocalizationTargetName, const FStrin
 	{
 		FString SourceFilename;
 		FString DestFilename;
-		TUniquePtr<FPortableObjectFormatDOM> ModifiedPOFile;
 		bool bMergeFailed = false;
 	};
 
+	LocFileNotifies.BeginParallelTasks(); // Batch the below SCC operations
+	ON_SCOPE_EXIT
+	{
+		LocFileNotifies.EndParallelTasks(); // End the SCC batch
+	};
+
 	TArray<FPOFileToMerge> POFilesToMerge;
 	for (const TTuple<FString, uint64>& SourceLocalizationFileHashPair : SourceLocalizationFileHashes)
 	{
@@ -1018,10 +1044,10 @@ bool MergeLocalization(const FString& SourceLocalizationTargetName, const FStrin
 		else
 		{
 			// File doesn't exist at dest; just copy directly from source
-			PreWriteFileWithSCC(DestFilename, bUseSourceControl);
+			LocFileNotifies.PreFileWrite(DestFilename);
 			if (IFileManager::Get().Copy(*DestFilename, *SourceFilename) == COPY_OK)
 			{
-				PostWriteFileWithSCC(DestFilename, bUseSourceControl);
+				LocFileNotifies.PostFileWrite(DestFilename);
 				UE_LOG(LogUGCLocalization, Log, TEXT("Merge copied '%s' to '%s'"), *SourceFilename, *DestFilename);
 			}
 			else
@@ -1032,13 +1058,13 @@ bool MergeLocalization(const FString& SourceLocalizationTargetName, const FStrin
 		}
 	}
 
-	ParallelFor(TEXT("UserGeneratedContentLocalization.MergeLocalization"), POFilesToMerge.Num(), 1, [&POFilesToMerge, &LoadPOFile, &SavePOFile, MergeMode, bUseSourceControl](int32 Index)
+	ParallelFor(TEXT("UserGeneratedContentLocalization.MergeLocalization"), POFilesToMerge.Num(), 1, [&POFilesToMerge, &LoadPOFile, &SavePOFile, MergeMode](int32 Index)
 	{
 		FPOFileToMerge& POFileToMerge = POFilesToMerge[Index];
 
 		FPortableObjectFormatDOM SourcePOFile;
-		TUniquePtr<FPortableObjectFormatDOM> DestPOFile = MakeUnique<FPortableObjectFormatDOM>();
-		if (!LoadPOFile(POFileToMerge.SourceFilename, SourcePOFile) || !LoadPOFile(POFileToMerge.DestFilename, *DestPOFile))
+		FPortableObjectFormatDOM DestPOFile;
+		if (!LoadPOFile(POFileToMerge.SourceFilename, SourcePOFile) || !LoadPOFile(POFileToMerge.DestFilename, DestPOFile))
 		{
 			POFileToMerge.bMergeFailed = true;
 			return;
@@ -1054,7 +1080,7 @@ bool MergeLocalization(const FString& SourceLocalizationTargetName, const FStrin
 				continue;
 			}
 
-			if (TSharedPtr<FPortableObjectEntry> DestPOEntry = DestPOFile->FindEntry(SourcePOEntry.ToSharedRef()))
+			if (TSharedPtr<FPortableObjectEntry> DestPOEntry = DestPOFile.FindEntry(SourcePOEntry.ToSharedRef()))
 			{
 				// Replace the entry in dest?
 				if (MergeMode == EMergeLocalizationMode::All || DestPOEntry->MsgStr.Num() == 0 || DestPOEntry->MsgStr[0].IsEmpty())
@@ -1066,42 +1092,23 @@ bool MergeLocalization(const FString& SourceLocalizationTargetName, const FStrin
 			else
 			{
 				// Add a new entry to dest
-				DestPOFile->AddEntry(MakeShared<FPortableObjectEntry>(*SourcePOEntry));
+				DestPOFile.AddEntry(MakeShared<FPortableObjectEntry>(*SourcePOEntry));
 				bModifiedDestPOFile = true;
 			}
 		}
 
 		if (bModifiedDestPOFile)
 		{
-			if (bUseSourceControl)
-			{
-				// When using source control, we need to defer the save until after the ParallelFor
-				POFileToMerge.ModifiedPOFile = MoveTemp(DestPOFile);
-			}
-			else
+			if (!SavePOFile(POFileToMerge.DestFilename, DestPOFile))
 			{
-				// When not using source control, we can handle the save within the ParallelFor
-				if (!SavePOFile(POFileToMerge.DestFilename, *DestPOFile))
-				{
-					POFileToMerge.bMergeFailed = true;
-					return;
-				}
+				POFileToMerge.bMergeFailed = true;
+				return;
 			}
 		}
 	});
 
 	for (FPOFileToMerge& POFileToMerge : POFilesToMerge)
 	{
-		// When using source control, we need to defer the save until after the ParallelFor
-		if (POFileToMerge.ModifiedPOFile)
-		{
-			check(bUseSourceControl);
-			if (!SavePOFile(POFileToMerge.DestFilename, *POFileToMerge.ModifiedPOFile))
-			{
-				POFileToMerge.bMergeFailed = true;
-			}
-		}
-
 		if (POFileToMerge.bMergeFailed)
 		{
 			return false;
@@ -1210,6 +1217,14 @@ void CleanupLocalization(TArrayView<const TSharedRef<IPlugin>> PluginsToClean, T
 			// Remove any leftover localization references from the plugins
 			auto RemovePluginLocalizationReferences = [](TArrayView<const TSharedRef<IPlugin>> Plugins)
 			{
+				FLocFileSCCNotifies LocFileNotifies = CreateLocFileNotifies();
+
+				LocFileNotifies.BeginParallelTasks(); // Batch the below SCC operations
+				ON_SCOPE_EXIT
+				{
+					LocFileNotifies.EndParallelTasks(); // End the SCC batch
+				};
+
 				for (const TSharedRef<IPlugin>& Plugin : Plugins)
 				{
 					const FString PluginLocalizationTargetName = GetLocalizationTargetName(Plugin);
@@ -1218,10 +1233,10 @@ void CleanupLocalization(TArrayView<const TSharedRef<IPlugin>> PluginsToClean, T
 					if (PluginDescriptor.LocalizationTargets.RemoveAll([&PluginLocalizationTargetName](const FLocalizationTargetDescriptor& LocalizationTargetDescriptor) { return LocalizationTargetDescriptor.Name == PluginLocalizationTargetName; }) > 0)
 					{
 						FText DescriptorUpdateFailureReason;
-						PreWriteFileWithSCC(Plugin->GetDescriptorFileName());
+						LocFileNotifies.PreFileWrite(Plugin->GetDescriptorFileName());
 						if (Plugin->UpdateDescriptor(PluginDescriptor, DescriptorUpdateFailureReason))
 						{
-							PostWriteFileWithSCC(Plugin->GetDescriptorFileName());
+							LocFileNotifies.PostFileWrite(Plugin->GetDescriptorFileName());
 						}
 					}
 				}
diff --git a/Engine/Source/Developer/Localization/Public/UserGeneratedContentLocalization.h b/Engine/Source/Developer/Localization/Public/UserGeneratedContentLocalization.h
--- a/Engine/Source/Developer/Localization/Public/UserGeneratedContentLocalization.h
+++ b/Engine/Source/Developer/Localization/Public/UserGeneratedContentLocalization.h
@@ -19,6 +19,7 @@
 
 class IPlugin;
 class FJsonObject;
+class FLocFileSCCNotifies;
 class FLocTextHelper;
 
 /**
@@ -242,6 +243,11 @@ LOCALIZATION_API void PreWriteFileWithSCC(const FString& Filename, const bool bU
  */
 LOCALIZATION_API void PostWriteFileWithSCC(const FString& Filename, const bool bUseSourceControl = true);
 
+/**
+ * Create a loc file notifies helper that will use source control if currently enabled.
+ */
+LOCALIZATION_API FLocFileSCCNotifies CreateLocFileNotifies(const bool bUseSourceControl = true);
+
 /**
  * Export UGC localization for the given plugins.
  * 
