Deprecated PreWriteFileWithSCC and PostWriteFileWithSCC in favor of CreateLocFileNotifies

CreateLocFileNotifies can batch SCC requests via BeginParallelTasks/EndParallelTasks, which can perform better than the single-file versions due to using fewer SCC queries

diff --git a/Engine/Source/Developer/Localization/Public/UserGeneratedContentLocalization.h b/Engine/Source/Developer/Localization/Public/UserGeneratedContentLocalization.h
--- a/Engine/Source/Developer/Localization/Public/UserGeneratedContentLocalization.h
+++ b/Engine/Source/Developer/Localization/Public/UserGeneratedContentLocalization.h
@@ -236,11 +236,13 @@ LOCALIZATION_API FString GetLocalizationTargetPOFile(const FString& Localization
 /**
  * Utility to prepare a SCC managed file for writing.
  */
+UE_DEPRECATED(5.8, "PreWriteFileWithSCC is deprecated in favor of CreateLocFileNotifies, as the latter can batch SCC requests via BeginParallelTasks/EndParallelTasks")
 LOCALIZATION_API void PreWriteFileWithSCC(const FString& Filename, const bool bUseSourceControl = true);
 
 /**
  * Utility to update a SCC managed file after writing.
  */
+UE_DEPRECATED(5.8, "PostWriteFileWithSCC is deprecated in favor of CreateLocFileNotifies, as the latter can batch SCC requests via BeginParallelTasks/EndParallelTasks")
 LOCALIZATION_API void PostWriteFileWithSCC(const FString& Filename, const bool bUseSourceControl = true);
 
 /**
