Ensure WriteFileIfModified always writes the file if it doesn't currently exist, even if the hash matches (empty file)

#rnx

diff --git a/Engine/Source/Developer/Localization/Private/LocalizationFileUtil.cpp b/Engine/Source/Developer/Localization/Private/LocalizationFileUtil.cpp
--- a/Engine/Source/Developer/Localization/Private/LocalizationFileUtil.cpp
+++ b/Engine/Source/Developer/Localization/Private/LocalizationFileUtil.cpp
@@ -32,7 +32,7 @@ bool WriteFileIfModified(const FString& Filename, TFunctionRef<bool(const FStrin
 	const uint64 CurrentHash = HashFunc(Filename);
 	const uint64 NewHash = HashFunc(TempFilename);
 
-	if (CurrentHash == NewHash)
+	if (CurrentHash == NewHash && FileManager.FileExists(*Filename))
 	{
 		return true;
 	}
