Fixed potential null pointer access

#rnx

diff --git a/Engine/Source/Developer/Localization/Private/TextLocalizationResourceGenerator.cpp b/Engine/Source/Developer/Localization/Private/TextLocalizationResourceGenerator.cpp
--- a/Engine/Source/Developer/Localization/Private/TextLocalizationResourceGenerator.cpp
+++ b/Engine/Source/Developer/Localization/Private/TextLocalizationResourceGenerator.cpp
@@ -107,6 +107,15 @@ bool FTextLocalizationResourceGenerator::GenerateLocRes(const FLocTextHelper& In
 	const bool bIsNativeCulture = InCultureToGenerate == InLocTextHelper.GetNativeCulture();
 	FCulturePtr Culture = FInternationalization::Get().GetCulture(InCultureToGenerate);
 
+	if (!Culture)
+	{
+		UE_LOGFMT(LogTextLocalizationResourceGenerator, Warning, "'{cultureCode}' did not resolve to a known culture",
+			("cultureCode", InCultureToGenerate),
+			("id", TextLocalizationResourceGenerator::LocalizationLogIdentifier)
+		);
+		return false;
+	}
+
 	TArray<FString> InheritedCultures = Culture->GetPrioritizedParentCultureNames();
 	InheritedCultures.Remove(Culture->GetName());
 
@@ -154,7 +163,7 @@ bool FTextLocalizationResourceGenerator::GenerateLocRes(const FLocTextHelper& In
 			if (bIsTranslated)
 			{
 				// Validate translations that look like they could be format patterns
-				if (EnumHasAnyFlags(InGenerateFlags, EGenerateLocResFlags::ValidateFormatPatterns) && Culture && TranslationText.Text.Contains(TEXT("{"), ESearchCase::CaseSensitive))
+				if (EnumHasAnyFlags(InGenerateFlags, EGenerateLocResFlags::ValidateFormatPatterns) && TranslationText.Text.Contains(TEXT("{"), ESearchCase::CaseSensitive))
 				{
 					const FTextFormat FmtPattern = FTextFormat::FromString(TranslationText.Text);
 
