Updated FLocalizationSCC to properly add files to the CheckOutFiles list after they're checked out in EndParallelTasks. Previously, if all files were successfully checked out, none would be moved from DeferredCheckedOutFiles to CheckedOutFiles.

#Synced-CL 49061450

diff --git a/Engine/Source/Developer/Localization/Private/LocalizationSourceControlUtil.cpp b/Engine/Source/Developer/Localization/Private/LocalizationSourceControlUtil.cpp
--- a/Engine/Source/Developer/Localization/Private/LocalizationSourceControlUtil.cpp
+++ b/Engine/Source/Developer/Localization/Private/LocalizationSourceControlUtil.cpp
@@ -131,6 +131,13 @@ bool FLocalizationSCC::EndParallelTasks(FText& OutError)
 		}
 		return false;
 	}
+	else
+	{
+		for (FString& DeferredCheckedOutFile : DeferredCheckedOutFilesArray)
+		{
+			CheckedOutFiles.Add(MoveTemp(DeferredCheckedOutFile));
+		}
+	}
 
 	return true;
 }
