Implement additional LocTextHelper constructor to accept pre-loaded localization manifest

#tests Create a new FInternationalizationManifest instance and load Engine/Content/Localization/Editor/Editor.manifest file with FJsonInternationalizationManifestSerializer::DeserializeManifestFromFile(); construct FLocTextHelper using this FInternationalizationManifest object, then try invoke FLocTextHelper::EnumerateSourceTextsParallel or FLocTextHelper::AddSourceText; observe everything is works the same way as if used one of the existing FLocTextHelper constructors with the following FLocTextHelper::LoadManifest.

#Synced-CL 50209799

diff --git a/Engine/Source/Developer/Localization/Private/LocTextHelper.cpp b/Engine/Source/Developer/Localization/Private/LocTextHelper.cpp
--- a/Engine/Source/Developer/Localization/Private/LocTextHelper.cpp
+++ b/Engine/Source/Developer/Localization/Private/LocTextHelper.cpp
@@ -500,6 +500,19 @@ FLocTextHelper::FLocTextHelper(FString InTargetPath, FString InManifestName, FSt
 	}
 }
 
+FLocTextHelper::FLocTextHelper(TSharedPtr<FInternationalizationManifest> InLoadedManifest, FString InTargetName, FString InTargetPath, FString InNativeCulture, TArray<FString> InForeignCultures, TSharedPtr<ILocFileNotifies> InLocFileNotifies, const ELocTextPlatformSplitMode InPlatformSplitMode)
+	: FLocTextHelper(MoveTemp(InTargetPath),
+		InTargetName + TEXT(".manifest"),
+		InTargetName + TEXT(".archive"),
+		MoveTemp(InNativeCulture),
+		MoveTemp(InForeignCultures),
+		MoveTemp(InLocFileNotifies),
+		InPlatformSplitMode)
+{
+	TargetName = MoveTemp(InTargetName);
+	Manifest = MoveTemp(InLoadedManifest);
+}
+
 bool FLocTextHelper::ShouldSplitPlatformData() const
 {
 	return FLocTextPlatformSplitUtils::ShouldSplitPlatformData(PlatformSplitMode);
diff --git a/Engine/Source/Developer/Localization/Public/LocTextHelper.h b/Engine/Source/Developer/Localization/Public/LocTextHelper.h
--- a/Engine/Source/Developer/Localization/Public/LocTextHelper.h
+++ b/Engine/Source/Developer/Localization/Public/LocTextHelper.h
@@ -297,6 +297,20 @@ public:
 	 */
 	UE_API FLocTextHelper(FString InTargetPath, FString InManifestName, FString InArchiveName, FString InNativeCulture, TArray<FString> InForeignCultures, TSharedPtr<ILocFileNotifies> InLocFileNotifies, const ELocTextPlatformSplitMode InPlatformSplitMode = ELocTextPlatformSplitMode::None);
 
+	/**
+	 * Initialize helper with already loaded manifest.
+	 * @note Suitable for dealing with in-memory manifest (but not limited to).
+	 *
+	 * @param InLoadedManifest		A pre-loaded manifest for intrusive LocTextHelper initialization.
+	 * @param InTargetName			The name of the loc target. Used to infer ManifestName and ArchiveName.
+	 * @param InTargetPath			Path to the localization target the pre-loaded manifest refers or should be saved to.
+	 * @param InNativeCulture		Culture code of the native culture (eg, en), or an empty string if the native culture is unknown.
+	 * @param InForeignCultures		Array of culture codes for the foreign cultures (the native culture will be removed from this array if present).
+	 * @param InLocFileNotifies		Interface for allowing source control integration (may be null).
+	 * @param InPlatformSplitMode	Should we split localization data per-platform?
+	 */
+	UE_API FLocTextHelper(TSharedPtr<FInternationalizationManifest> InLoadedManifest, FString InTargetName, FString InTargetPath, FString InNativeCulture, TArray<FString> InForeignCultures, TSharedPtr<ILocFileNotifies> InLocFileNotifies, const ELocTextPlatformSplitMode InPlatformSplitMode = ELocTextPlatformSplitMode::None);
+
 	/**
 	 * Are we splitting localization data per-platform?
 	 */
