Added RapidJSON utils for reading a string value into a FStringView

#Synced-CL 51535417

diff --git a/Engine/Source/Developer/Localization/Private/Serialization/JsonInternationalizationMetadataSerializer.cpp b/Engine/Source/Developer/Localization/Private/Serialization/JsonInternationalizationMetadataSerializer.cpp
--- a/Engine/Source/Developer/Localization/Private/Serialization/JsonInternationalizationMetadataSerializer.cpp
+++ b/Engine/Source/Developer/Localization/Private/Serialization/JsonInternationalizationMetadataSerializer.cpp
@@ -32,7 +32,7 @@ TSharedPtr<FLocMetadataValue> JsonValueToLocMetaDataValue(const UE::Json::FValue
 
 		case rapidjson::kStringType:
 		{
-			return MakeShared<FLocMetadataValueString>(FString(FStringView(JsonValue.GetString(), JsonValue.GetStringLength())));
+			return MakeShared<FLocMetadataValueString>(FString(UE::Json::ValueAsStringView(JsonValue)));
 		}
 
 		case rapidjson::kArrayType:
@@ -61,7 +61,7 @@ TSharedPtr<FLocMetadataValue> JsonValueToLocMetaDataValue(const UE::Json::FValue
 			{
 				if (TSharedPtr<FLocMetadataValue> MetadataValue = JsonValueToLocMetaDataValue(ValueIter.value))
 				{
-					MetadataSubObject->SetField(FString(FStringView(ValueIter.name.GetString(), ValueIter.name.GetStringLength())), MetadataValue);
+					MetadataSubObject->SetField(FString(UE::Json::ValueAsStringView(ValueIter.name)), MetadataValue);
 				}
 			}
 			return MakeShared<FLocMetadataValueObject>(MoveTemp(MetadataSubObject));
diff --git a/Engine/Source/Runtime/Json/Private/JsonUtils/RapidJsonUtils.cpp b/Engine/Source/Runtime/Json/Private/JsonUtils/RapidJsonUtils.cpp
--- a/Engine/Source/Runtime/Json/Private/JsonUtils/RapidJsonUtils.cpp
+++ b/Engine/Source/Runtime/Json/Private/JsonUtils/RapidJsonUtils.cpp
@@ -8,6 +8,15 @@ namespace UE::Json
 
 const constexpr uint32 DefaultParseFlags = rapidjson::ParseFlag::kParseTrailingCommasFlag;
 
+TOptional<FStringView> GetStringValue(const FValue& Value)
+{
+	if (Value.IsString())
+	{
+		return ValueAsStringView(Value);
+	}
+	return {};
+}
+
 TOptional<bool> GetBoolField(FConstObject Object, const TCHAR* FieldName)
 {
 	FValue::ConstMemberIterator It = Object.FindMember(FieldName);
@@ -83,7 +92,7 @@ TOptional<FStringView> GetStringField(FConstObject Object, const TCHAR* FieldNam
 		return {};
 	}
 
-	return FStringView(It->value.GetString(), It->value.GetStringLength());
+	return ValueAsStringView(It->value);
 }
 
 bool HasField(FConstObject Object, const TCHAR* FieldName)
diff --git a/Engine/Source/Runtime/Json/Public/JsonUtils/RapidJsonUtils.h b/Engine/Source/Runtime/Json/Public/JsonUtils/RapidJsonUtils.h
--- a/Engine/Source/Runtime/Json/Public/JsonUtils/RapidJsonUtils.h
+++ b/Engine/Source/Runtime/Json/Public/JsonUtils/RapidJsonUtils.h
@@ -89,6 +89,23 @@ namespace Private
 		return FValue(Str.GetData(), Str.Len(), Allocator);
 	}
 
+	/**
+	 * Reads a string from the specified value.
+	 * @note The given value must be a string.
+	 */
+	inline FStringView ValueAsStringView(const FValue& Value)
+	{
+		check(Value.IsString());
+		return FStringView(Value.GetString(), Value.GetStringLength());
+	}
+
+	/**
+	 * Attempts to read a string from the specified value.
+	 *
+	 * Returns the value, if it matches the type explicitly
+	 */
+	JSON_API TOptional<FStringView> GetStringValue(const FValue& Value);
+
 	/**
 	 * Attempts to read a boolean with a specified name.
 	 * 
