Fixed MergeLocalization during ExportLocalization repeatedly merging in stale localization data

It now checks that the timestamp of the source file is newer than the timestamp of the scratch file, to ensure it's not constantly re-merging potentially stale data that has already been processed and then potentially replaced by a fresh export

#Synced-CL 52493164

#Synced-CL 52465190

diff --git a/Engine/Source/Developer/Localization/Private/UserGeneratedContentLocalization.cpp b/Engine/Source/Developer/Localization/Private/UserGeneratedContentLocalization.cpp
--- a/Engine/Source/Developer/Localization/Private/UserGeneratedContentLocalization.cpp
+++ b/Engine/Source/Developer/Localization/Private/UserGeneratedContentLocalization.cpp
@@ -338,15 +338,17 @@ bool ExportLocalization(TArrayView<const TSharedRef<IPlugin>> Plugins, const FEx
 		return false;
 	}
 
+	IFileManager& FileManager = IFileManager::Get();
+
 	// Create a scratch directory for the temporary localization data
 	const FString RootLocalizationScratchDirectory = GetLocalizationScratchDirectory();
-	IFileManager::Get().MakeDirectory(*RootLocalizationScratchDirectory, /*bTree*/true);
+	FileManager.MakeDirectory(*RootLocalizationScratchDirectory, /*bTree*/true);
 	ON_SCOPE_EXIT
 	{
 		if (ExportOptions.bAutoCleanup)
 		{
 			// Delete the entire scratch directory
-			IFileManager::Get().DeleteDirectory(*RootLocalizationScratchDirectory, /*bRequireExists*/false, /*bTree*/true);
+			FileManager.DeleteDirectory(*RootLocalizationScratchDirectory, /*bRequireExists*/false, /*bTree*/true);
 		}
 	};
 
@@ -374,9 +376,11 @@ bool ExportLocalization(TArrayView<const TSharedRef<IPlugin>> Plugins, const FEx
 			FString DestinationFilename = SourceFilename;
 			if (DestinationFilename.ReplaceInline(*PluginLocalizationTargetDirectory, *PluginLocalizationScratchDirectory) > 0)
 			{
-				if (IFileManager::Get().Copy(*DestinationFilename, *SourceFilename) == COPY_OK)
+				if (FileManager.Copy(*DestinationFilename, *SourceFilename) == COPY_OK)
 				{
-					UE_LOG(LogUGCLocalization, Display, TEXT("Imported existing file for '%s': %s"), *PluginLocalizationTargetName, *SourceFilename);
+					// Preserve the source timestamp, as MergeLocalization relies on this
+					FileManager.SetTimeStamp(*DestinationFilename, FileManager.GetTimeStamp(*SourceFilename));
+					UE_LOGF(LogUGCLocalization, Display, "Imported existing file for '%ls': %ls", *PluginLocalizationTargetName, *SourceFilename);
 				}
 				else
 				{
@@ -388,7 +392,7 @@ bool ExportLocalization(TArrayView<const TSharedRef<IPlugin>> Plugins, const FEx
 		if (!ExportOptions.LocalizationTargetRootDirectoryOverride.IsEmpty() && ExportOptions.MergeProjectDataWithRootDirectoryOverrideData.IsSet())
 		{
 			const FString PluginLocalizationSourceDirectory = GetLocalizationTargetDirectory(PluginLocalizationTargetName, Plugin->GetContentDir());
-			if (!MergeLocalization(PluginLocalizationTargetName, PluginLocalizationSourceDirectory, PluginLocalizationTargetName, PluginLocalizationScratchDirectory, ExportOptions.MergeProjectDataWithRootDirectoryOverrideData.GetValue(), /*bUseSourceControl*/false))
+			if (!MergeLocalization(PluginLocalizationTargetName, PluginLocalizationSourceDirectory, PluginLocalizationTargetName, PluginLocalizationScratchDirectory, ExportOptions.MergeProjectDataWithRootDirectoryOverrideData.GetValue(), /*bOnlyIfSourceFileNewer*/true, /*bUseSourceControl*/false))
 			{
 				return false;
 			}
@@ -656,7 +660,7 @@ bool ExportLocalization(TArrayView<const TSharedRef<IPlugin>> Plugins, const FEx
 					!SourceFileHash || *SourceFileHash != LocalizationFileHashPair.Value)
 				{
 					LocFileNotifies.PreFileWrite(DestinationFilename);
-					if (IFileManager::Get().Copy(*DestinationFilename, *SourceFilename) == COPY_OK)
+					if (FileManager.Copy(*DestinationFilename, *SourceFilename) == COPY_OK)
 					{
 						LocFileNotifies.PostFileWrite(DestinationFilename);
 						UE_LOG(LogUGCLocalization, Display, TEXT("Updated file for '%s': %s"), *PluginLocalizationTargetName, *DestinationFilename);
@@ -956,7 +960,7 @@ ELoadLocalizationResult LoadLocalization(const FString& LocalizationTargetName,
 	return ELoadLocalizationResult::Success;
 }
 
-bool MergeLocalization(const FString& SourceLocalizationTargetName, const FString& SourceLocalizationTargetDirectory, const FString& DestLocalizationTargetName, const FString& DestLocalizationTargetDirectory, const EMergeLocalizationMode MergeMode, const bool bUseSourceControl)
+bool MergeLocalization(const FString& SourceLocalizationTargetName, const FString& SourceLocalizationTargetDirectory, const FString& DestLocalizationTargetName, const FString& DestLocalizationTargetDirectory, const EMergeLocalizationMode MergeMode, const bool bOnlyIfSourceFileNewer, const bool bUseSourceControl)
 {
 	TMap<FString, uint64> SourceLocalizationFileHashes;
 	GetLocalizationFileHashes(SourceLocalizationTargetDirectory, SourceLocalizationFileHashes);
@@ -970,6 +974,7 @@ bool MergeLocalization(const FString& SourceLocalizationTargetName, const FStrin
 	TMap<FString, uint64> DestLocalizationFileHashes;
 	GetLocalizationFileHashes(DestLocalizationTargetDirectory, DestLocalizationFileHashes);
 
+	IFileManager& FileManager = IFileManager::Get();
 	FLocFileSCCNotifies LocFileNotifies = CreateLocFileNotifies(bUseSourceControl);
 
 	auto SourceFilenameToDestFilename = [&SourceLocalizationTargetDirectory, &DestLocalizationTargetDirectory, SourcePOFilename = SourceLocalizationTargetName + TEXT(".po"), DestPOFilename = DestLocalizationTargetName + TEXT(".po")](const FString& SourceFilename)
@@ -1033,9 +1038,23 @@ bool MergeLocalization(const FString& SourceLocalizationTargetName, const FStrin
 			if (SourceLocalizationFileHashPair.Value == *DestLocalizationFileHash)
 			{
 				// File hash is identical; nothing to merge
-				UE_LOG(LogUGCLocalization, Log, TEXT("Merge skipped '%s' as it is identical to '%s'"), *SourceFilename, *DestFilename);
+				UE_LOGF(LogUGCLocalization, Log, "Merge skipped '%ls' as it is identical to '%ls'", *SourceFilename, *DestFilename);
+				continue;
+			}
+			
+			if (bOnlyIfSourceFileNewer)
+			{
+				const FDateTime SourceTimestamp = FileManager.GetTimeStamp(*SourceFilename);
+				const FDateTime DestTimestamp = FileManager.GetTimeStamp(*DestFilename);
+				if (SourceTimestamp <= DestTimestamp)
+				{
+					// Source file is older; nothing to merge
+					UE_LOGF(LogUGCLocalization, Log, "Merge skipped '%ls' as it is older than '%ls'", *SourceFilename, *DestFilename);
+					continue;
+				}
 			}
-			else if (FPathViews::GetExtension(SourceFilename) == TEXTVIEW("po"))
+			
+			if (FPathViews::GetExtension(SourceFilename) == TEXTVIEW("po"))
 			{
 				// Merge source into dest
 				POFilesToMerge.Add(FPOFileToMerge{ SourceFilename, DestFilename });
@@ -1045,7 +1064,7 @@ bool MergeLocalization(const FString& SourceLocalizationTargetName, const FStrin
 		{
 			// File doesn't exist at dest; just copy directly from source
 			LocFileNotifies.PreFileWrite(DestFilename);
-			if (IFileManager::Get().Copy(*DestFilename, *SourceFilename) == COPY_OK)
+			if (FileManager.Copy(*DestFilename, *SourceFilename) == COPY_OK)
 			{
 				LocFileNotifies.PostFileWrite(DestFilename);
 				UE_LOG(LogUGCLocalization, Log, TEXT("Merge copied '%s' to '%s'"), *SourceFilename, *DestFilename);
diff --git a/Engine/Source/Developer/Localization/Public/UserGeneratedContentLocalization.h b/Engine/Source/Developer/Localization/Public/UserGeneratedContentLocalization.h
--- a/Engine/Source/Developer/Localization/Public/UserGeneratedContentLocalization.h
+++ b/Engine/Source/Developer/Localization/Public/UserGeneratedContentLocalization.h
@@ -300,9 +300,10 @@ LOCALIZATION_API ELoadLocalizationResult LoadLocalization(const FString& Localiz
  * @param DestLocalizationTargetName		The name of the dest localization target (@see GetLocalizationTargetName).
  * @param DestLocalizationTargetDirectory	The directory of the dest localization target (@see GetLocalizationTargetDirectory).
  * @param MergeMode							When should an existing translation from source be applied to dest?
+ * @param bOnlyIfSourceFileNewer			True if we should only merge from source -> dest if the source file is newer than the dest file.
  * @param bUseSourceControl					Pass false as an optimization if you knew for certain that DestLocalizationTargetDirectory is outside of source control, otherwise pass true.
  */
-LOCALIZATION_API bool MergeLocalization(const FString& SourceLocalizationTargetName, const FString& SourceLocalizationTargetDirectory, const FString& DestLocalizationTargetName, const FString& DestLocalizationTargetDirectory, const EMergeLocalizationMode MergeMode, const bool bUseSourceControl = true);
+LOCALIZATION_API bool MergeLocalization(const FString& SourceLocalizationTargetName, const FString& SourceLocalizationTargetDirectory, const FString& DestLocalizationTargetName, const FString& DestLocalizationTargetDirectory, const EMergeLocalizationMode MergeMode, const bool bOnlyIfSourceFileNewer, const bool bUseSourceControl = true);
 
 /**
  * Cleanup UGC localization that is no longer relevant based on the given descriptor.
